Minor cleanup of SteamController module (removing printfs, etc).

diff --git a/Engine/Source/Runtime/Linux/SteamController/Private/SteamControllerInterface.cpp b/Engine/Source/Runtime/Linux/SteamController/Private/SteamControllerInterface.cpp
--- a/Engine/Source/Runtime/Linux/SteamController/Private/SteamControllerInterface.cpp
+++ b/Engine/Source/Runtime/Linux/SteamController/Private/SteamControllerInterface.cpp
@@ -12,6 +12,7 @@
 #define RIGHT_THUMBPAD_DEADZONE 0
 
 IMPLEMENT_MODULE( FDefaultModuleImpl, SteamController );
+DEFINE_LOG_CATEGORY(LogSteamController);
 
 TSharedRef< SteamControllerInterface > SteamControllerInterface::Create(  const TSharedRef< FGenericApplicationMessageHandler >& InMessageHandler )
 {
@@ -23,16 +24,19 @@ SteamControllerInterface::SteamControllerInterface( const TSharedRef< FGenericAp
 {
 	if(SteamController() != NULL)
 	{
-		printf("MichaelL - SteamController Init!\n");
 		FString vdfPath = FPaths::ConvertRelativePathToFull(FPaths::Combine(*FPaths::GeneratedConfigDir(), TEXT("controller.vdf")));
-		printf("vdfPath: %s\n", TCHAR_TO_ANSI(*vdfPath));
-		SteamController()->Init( TCHAR_TO_ANSI(*vdfPath) );
+		bool bInited = SteamController()->Init(TCHAR_TO_ANSI(*vdfPath));
+		UE_LOG(LogSteamController, Log, TEXT("SteamController %s initialized with vdf file '%s'."), bInited ? TEXT("could not be") : TEXT("has been"), *vdfPath);
 
-		memset(ControllerStates, 0, sizeof(FControllerState) * MAX_STEAM_CONTROLLERS);
+		// [RCL] 2014-05-05 FIXME: disable when could not init?
+		if (bInited)
+		{
+			FMemory::MemZero(ControllerStates);
+		}
 	}
 	else
 	{
-		printf("MichaelL - SteamController failed...\n");
+		UE_LOG(LogSteamController, Log, TEXT("SteamController is not available"));
 	}
 
 	InitialButtonRepeatDelay = 0.2f;
diff --git a/Engine/Source/Runtime/Linux/SteamController/Public/SteamControllerInterface.h b/Engine/Source/Runtime/Linux/SteamController/Public/SteamControllerInterface.h
--- a/Engine/Source/Runtime/Linux/SteamController/Public/SteamControllerInterface.h
+++ b/Engine/Source/Runtime/Linux/SteamController/Public/SteamControllerInterface.h
@@ -12,6 +12,8 @@
 /** Max number of controller buttons.  Must be < 256*/
 #define MAX_NUM_CONTROLLER_BUTTONS 26
 
+DECLARE_LOG_CATEGORY_EXTERN(LogSteamController, Log, All);
+
 /**
  * Interface class for XInput devices (xbox 360 controller)                 
  */
