Support compressed audio on Linux.

- Catching up with MatthewG's changes on other platforms.
- Fixes problems on QA-SoundLoop and elsewhere (TTP #334227).

diff --git a/Engine/Source/Runtime/Engine/Private/AudioDecompress.cpp b/Engine/Source/Runtime/Engine/Private/AudioDecompress.cpp
--- a/Engine/Source/Runtime/Engine/Private/AudioDecompress.cpp
+++ b/Engine/Source/Runtime/Engine/Private/AudioDecompress.cpp
@@ -33,7 +33,7 @@
 #define VORBIS_BYTE_ORDER 1
 #endif
 
-#define WITH_OPUS (PLATFORM_WINDOWS || PLATFORM_MAC)
+#define WITH_OPUS (PLATFORM_WINDOWS || PLATFORM_MAC || PLATFORM_LINUX)
 
 #if WITH_OPUS
 #include "opus.h"
diff --git a/Engine/Source/Runtime/Linux/ALAudio/Private/ALAudioBuffer.cpp b/Engine/Source/Runtime/Linux/ALAudio/Private/ALAudioBuffer.cpp
--- a/Engine/Source/Runtime/Linux/ALAudio/Private/ALAudioBuffer.cpp
+++ b/Engine/Source/Runtime/Linux/ALAudio/Private/ALAudioBuffer.cpp
@@ -91,6 +91,21 @@ FALSoundBuffer* FALSoundBuffer::CreateNativeBuffer( FALAudioDevice* AudioDevice,
 {
 	SCOPE_CYCLE_COUNTER( STAT_AudioResourceCreationTime );
 
+	// Check to see if thread has finished decompressing on the other thread
+	if (Wave->AudioDecompressor != NULL)
+	{
+		if (!Wave->AudioDecompressor->IsDone())
+		{
+			// Don't play this sound just yet
+			UE_LOG(LogAudio, Log, TEXT("Waiting for sound to decompress: %s"), *Wave->GetName());
+			return(NULL);
+		}
+
+		// Remove the decompressor
+		delete Wave->AudioDecompressor;
+		Wave->AudioDecompressor = NULL;
+	}
+
 	// Can't create a buffer without any source data
 	if( Wave == NULL || Wave->NumChannels == 0 )
 	{
diff --git a/Engine/Source/Runtime/Linux/ALAudio/Private/ALAudioDevice.cpp b/Engine/Source/Runtime/Linux/ALAudio/Private/ALAudioDevice.cpp
--- a/Engine/Source/Runtime/Linux/ALAudio/Private/ALAudioDevice.cpp
+++ b/Engine/Source/Runtime/Linux/ALAudio/Private/ALAudioDevice.cpp
@@ -14,6 +14,7 @@
 #include "AudioEffect.h"
 #include "Engine.h"
 
+DEFINE_LOG_CATEGORY(LogALAudio);
 
 //     2 UU == 1"
 // <=> 1 UU == 0.0127 m
@@ -109,15 +110,13 @@ bool FALAudioDevice::InitializeHardware( void )
 	HardwareDevice = alcOpenDevice(  NULL );
 	if( !HardwareDevice )
 	{
-		//UE_LOG(LogALAudio, Log, TEXT( "ALAudio: no OpenAL devices found." ) );
+		UE_LOG(LogALAudio, Log, TEXT( "ALAudio: no OpenAL devices found." ) );
 		return  false  ;
 	}
 
 	// Display the audio device that was actually opened
 	const ALCchar* OpenedDeviceName = alcGetString( HardwareDevice, ALC_DEVICE_SPECIFIER );
-	//UE_LOG(LogALAudio, Log, TEXT( "ALAudio: no OpenAL devices found." ) );
-	//UE_LOG(LogALAudio, Log, TEXT( "ALAudio device requested : %s" ), *DeviceName );
-	//UE_LOG(LogALAudio, Log, TEXT( "ALAudio device Opened : %s" ), OpenedDeviceName );
+	UE_LOG(LogALAudio, Log, TEXT("ALAudio device opened : %s"), StringCast<TCHAR>(static_cast<const ANSICHAR*>(OpenedDeviceName)).Get());
 
 	// Create a context
 	int Caps[] = 
@@ -143,14 +142,14 @@ bool FALAudioDevice::InitializeHardware( void )
 	// Make sure everything happened correctly
 	if( alError( TEXT( "Init" ) ) )
 	{
-//		debugf( NAME_Init, TEXT( "ALAudio: alcMakeContextCurrent failed." ) );
+		UE_LOG(LogALAudio, Warning, TEXT("ALAudio: alcMakeContextCurrent failed."));
 		return false ;
 	}
 
-// 	debugf( NAME_Init, TEXT( "AL_VENDOR      : %s" ), ANSI_TO_TCHAR( ( ANSICHAR* )alGetString( AL_VENDOR ) ) );
-// 	debugf( NAME_Init, TEXT( "AL_RENDERER    : %s" ), ANSI_TO_TCHAR( ( ANSICHAR* )alGetString( AL_RENDERER ) ) );
-// 	debugf( NAME_Init, TEXT( "AL_VERSION     : %s" ), ANSI_TO_TCHAR( ( ANSICHAR* )alGetString( AL_VERSION ) ) );
-// 	debugf( NAME_Init, TEXT( "AL_EXTENSIONS  : %s" ), ANSI_TO_TCHAR( ( ANSICHAR* )alGetString( AL_EXTENSIONS ) ) );
+	UE_LOG(LogALAudio, Log, TEXT("AL_VENDOR      : %s"), StringCast<TCHAR>(static_cast<const ANSICHAR*>(alGetString(AL_VENDOR))).Get());
+	UE_LOG(LogALAudio, Log, TEXT("AL_RENDERER    : %s"), StringCast<TCHAR>(static_cast<const ANSICHAR*>(alGetString(AL_RENDERER))).Get());
+	UE_LOG(LogALAudio, Log, TEXT("AL_VERSION     : %s"), StringCast<TCHAR>(static_cast<const ANSICHAR*>(alGetString(AL_VERSION))).Get());
+	UE_LOG(LogALAudio, Log, TEXT("AL_EXTENSIONS  : %s"), StringCast<TCHAR>(static_cast<const ANSICHAR*>(alGetString(AL_EXTENSIONS))).Get());
  
 	// Get the enums for multichannel support
 #if !PLATFORM_HTML5
@@ -180,13 +179,13 @@ bool FALAudioDevice::InitializeHardware( void )
 
 	if( Sources.Num() < 1 )
 	{
-//		debugf( NAME_Error,TEXT( "ALAudio: couldn't allocate any sources" ) );
+		UE_LOG(LogALAudio, Warning, TEXT("ALAudio: couldn't allocate any sources"));
 		return false ;
 	}
 
 	// Update MaxChannels in case we couldn't create enough sources.
 	MaxChannels = Sources.Num();
-/*	debugf( NAME_Init, TEXT( "ALAudioDevice: Allocated %i sources" ), MaxChannels );*/
+	UE_LOG(LogALAudio, Verbose, TEXT("ALAudioDevice: Allocated %i sources"), MaxChannels);
 
 	// Use our own distance model.
 	alDistanceModel( AL_NONE );
@@ -197,7 +196,6 @@ bool FALAudioDevice::InitializeHardware( void )
 	// Initialized.
 	NextResourceID = 1;
 
-
 	return true ;
 }
 
@@ -365,12 +363,7 @@ OpenAL utility functions
 //
 bool FALAudioDevice::FindExt( const TCHAR* Name )
 {
-	if( alIsExtensionPresent( TCHAR_TO_ANSI( Name ) ) || alcIsExtensionPresent( HardwareDevice, TCHAR_TO_ANSI( Name ) ) )
-	{
-		//UE_LOG ( LogALAudio, log,TEXT( "Device supports: %s" )); 
-		return( true );
-	}
-	return( true );
+	return alIsExtensionPresent(TCHAR_TO_ANSI(Name)) || alcIsExtensionPresent(HardwareDevice, TCHAR_TO_ANSI(Name));
 }
 
 //
@@ -403,24 +396,24 @@ bool FALAudioDevice::alError( const TCHAR* Text, bool Log )
 			{
 				switch ( Error )
 				{
-				case AL_INVALID_NAME:
-					UE_LOG ( LogInit, Log, TEXT( "ALAudio: AL_INVALID_NAME in %s" ), Text ); 
-					break;
-				case AL_INVALID_VALUE:
-					UE_LOG ( LogInit, Log, TEXT( "ALAudio: AL_INVALID_VALUE in %s" ), Text ); 
-					break;
-				case AL_OUT_OF_MEMORY:
-					UE_LOG ( LogInit, Log, TEXT( "ALAudio: AL_OUT_OF_MEMORY in %s" ), Text ); 
-					break;
-				case AL_INVALID_ENUM:
-					UE_LOG ( LogInit, Log, TEXT( "ALAudio: AL_INVALID_ENUM in %s" ), Text ); 
-					break;
-				case AL_INVALID_OPERATION:
-					UE_LOG ( LogInit, Log, TEXT( "ALAudio: AL_INVALID_OPERATION in %s" ), Text ); 
-					break;
-				default:
-					UE_LOG ( LogInit, Log, TEXT( "ALAudio: Unknown Error NUM in %s" ), Text ); 
-					break;
+					case AL_INVALID_NAME:
+						UE_LOG(LogALAudio, Warning, TEXT("ALAudio: AL_INVALID_NAME in %s"), Text);
+						break;
+					case AL_INVALID_VALUE:
+						UE_LOG(LogALAudio, Warning, TEXT("ALAudio: AL_INVALID_VALUE in %s"), Text);
+						break;
+					case AL_OUT_OF_MEMORY:
+						UE_LOG(LogALAudio, Warning, TEXT("ALAudio: AL_OUT_OF_MEMORY in %s"), Text);
+						break;
+					case AL_INVALID_ENUM:
+						UE_LOG(LogALAudio, Warning, TEXT("ALAudio: AL_INVALID_ENUM in %s"), Text);
+						break;
+					case AL_INVALID_OPERATION:
+						UE_LOG(LogALAudio, Warning, TEXT("ALAudio: AL_INVALID_OPERATION in %s"), Text);
+						break;
+					default:
+						UE_LOG(LogALAudio, Warning, TEXT("ALAudio: Unknown Error NUM in %s"), Text);
+						break;
 				}
 			}
 		}
@@ -436,3 +429,21 @@ FSoundSource* FALAudioDevice::CreateSoundSource()
 {
 	return new FALSoundSource(this);
 }
+
+bool FALAudioDevice::HasCompressedAudioInfoClass(USoundWave* SoundWave)
+{
+#if WITH_OGGVORBIS
+	return true;
+#else
+	return false;
+#endif
+}
+
+class ICompressedAudioInfo* FALAudioDevice::CreateCompressedAudioInfo(USoundWave* SoundWave)
+{
+#if WITH_OGGVORBIS
+	return new FVorbisAudioInfo();
+#else
+	return NULL;
+#endif
+}
diff --git a/Engine/Source/Runtime/Linux/ALAudio/Public/ALAudioDevice.h b/Engine/Source/Runtime/Linux/ALAudio/Public/ALAudioDevice.h
--- a/Engine/Source/Runtime/Linux/ALAudio/Public/ALAudioDevice.h
+++ b/Engine/Source/Runtime/Linux/ALAudio/Public/ALAudioDevice.h
@@ -26,6 +26,8 @@ class FALAudioDevice;
 #pragma pack (pop)
 #endif 
 
+DECLARE_LOG_CATEGORY_EXTERN(LogALAudio, Log, All);
+
 /**
  * OpenAL implementation of FSoundBuffer, containing the wave data and format information.
  */
@@ -229,10 +231,14 @@ public:
 	 */
 	virtual void FreeResource( USoundWave* SoundWave );
 
+	virtual bool HasCompressedAudioInfoClass(USoundWave* SoundWave) OVERRIDE;
+
+	virtual class ICompressedAudioInfo* CreateCompressedAudioInfo(USoundWave* SoundWave) OVERRIDE;
+
 	void FindProcs( bool AllowExt );
 
 	// Error checking.
-	bool alError( const TCHAR* Text, bool Log = 1 );
+	bool alError( const TCHAR* Text, bool Log = true );
 
 protected:
 
