ALAudio: more logging to easily identify errors + some cleanup.

diff --git a/Engine/Source/Runtime/Linux/ALAudio/Private/ALAudioBuffer.cpp b/Engine/Source/Runtime/Linux/ALAudio/Private/ALAudioBuffer.cpp
--- a/Engine/Source/Runtime/Linux/ALAudio/Private/ALAudioBuffer.cpp
+++ b/Engine/Source/Runtime/Linux/ALAudio/Private/ALAudioBuffer.cpp
@@ -43,15 +43,15 @@ FALSoundBuffer::~FALSoundBuffer( void )
  * @return	FALSoundBuffer pointer if buffer creation succeeded, NULL otherwise
  */
 
-FALSoundBuffer* FALSoundBuffer::Init(  FALAudioDevice* AudioDevice ,USoundWave* InWave )
+FALSoundBuffer* FALSoundBuffer::Init(FALAudioDevice* AudioDevice, USoundWave* InWave)
 {
 	// Can't create a buffer without any source data
-	if (InWave == NULL || InWave->NumChannels == 0)
+	if (InWave == nullptr || InWave->NumChannels == 0)
 	{
-		return NULL;
+		return nullptr;
 	}
 
-	FALSoundBuffer *Buffer = NULL;
+	FALSoundBuffer *Buffer = nullptr;
 
 	switch (static_cast<EDecompressionType>(InWave->DecompressionType))
 	{
@@ -74,17 +74,28 @@ FALSoundBuffer* FALSoundBuffer::Init(  FALAudioDevice* AudioDevice ,USoundWave*
 		}
 		break;
 
+		// Always create a new buffer for streaming ogg vorbis data
+		//Buffer = CreateQueuedBuffer(AudioDevice, InWave);
+		//break;
+
 	case DTYPE_Invalid:
 	case DTYPE_Preview:
 	case DTYPE_Procedural:
 	case DTYPE_RealTime:
 	default:
 		// Invalid will be set if the wave cannot be played
+		UE_LOG(LogALAudio, Warning, TEXT("ALSoundBuffer wave '%s' has an invalid decompression type %d."), 
+			*InWave->GetName(), static_cast<EDecompressionType>(InWave->DecompressionType));
 		break;
 	}
 
-	return Buffer;
+	if (Buffer == nullptr)
+	{
+		UE_LOG(LogALAudio, Warning, TEXT("ALSoundBuffer init failed for wave '%s', decompression type %d."),
+			*InWave->GetName(), static_cast<EDecompressionType>(InWave->DecompressionType));
+	}
 
+	return Buffer;
 }
 
 FALSoundBuffer* FALSoundBuffer::CreateNativeBuffer( FALAudioDevice* AudioDevice, USoundWave* Wave)
@@ -92,53 +103,59 @@ FALSoundBuffer* FALSoundBuffer::CreateNativeBuffer( FALAudioDevice* AudioDevice,
 	SCOPE_CYCLE_COUNTER( STAT_AudioResourceCreationTime );
 
 	// Check to see if thread has finished decompressing on the other thread
-	if (Wave->AudioDecompressor != NULL)
+	if (Wave->AudioDecompressor)
 	{
 		Wave->AudioDecompressor->EnsureCompletion();
 
 		// Remove the decompressor
 		delete Wave->AudioDecompressor;
-		Wave->AudioDecompressor = NULL;
+		Wave->AudioDecompressor = nullptr;
 	}
 
 	// Can't create a buffer without any source data
-	if( Wave == NULL || Wave->NumChannels == 0 )
+	if (Wave == nullptr)
 	{
-		return( NULL );
+		return nullptr;
 	}
-	FWaveModInfo WaveInfo;
 
-	Wave->InitAudioResource(AudioDevice->GetRuntimeFormat());
+	if (Wave->NumChannels == 0)
+	{
+		UE_LOG(LogALAudio, Warning, TEXT("ALSoundBuffer: cannot create buffer for wave '%s' with 0 channels"),
+			*Wave->GetName());
+		return nullptr;
+	}
 
-	FALSoundBuffer* Buffer = NULL;
+	FWaveModInfo WaveInfo;
+	Wave->InitAudioResource(AudioDevice->GetRuntimeFormat());
+	FALSoundBuffer* Buffer = nullptr;
 	
 	// Find the existing buffer if any
-	if( Wave->ResourceID )
+	if (Wave->ResourceID)
 	{
 		Buffer = AudioDevice->WaveBufferMap.FindRef( Wave->ResourceID );
 	}
 
-	if( Buffer == NULL )
+	if (Buffer == nullptr)
 	{
 		// Create new buffer.
-		Buffer = new FALSoundBuffer( AudioDevice );
+		Buffer = new FALSoundBuffer(AudioDevice);
 
-		alGenBuffers( 1, Buffer->BufferIds );
+		alGenBuffers(1, Buffer->BufferIds);
 
-		AudioDevice->alError( TEXT( "RegisterSound" ) );
+		AudioDevice->alError(TEXT("RegisterSound"));
 
 		// Allocate new resource ID and assign to USoundNodeWave. A value of 0 (default) means not yet registered.
 		int ResourceID = AudioDevice->NextResourceID++;
 		Buffer->ResourceID = ResourceID;
 		Wave->ResourceID = ResourceID;
 
-		AudioDevice->Buffers.Add( Buffer );
-		AudioDevice->WaveBufferMap.FindOrAdd(ResourceID) =  Buffer ;
+		AudioDevice->Buffers.Add(Buffer);
+		AudioDevice->WaveBufferMap.FindOrAdd(ResourceID) =  Buffer;
 
 		// Keep track of associated resource name.
 		Buffer->ResourceName = Wave->GetPathName();
 
-		Buffer->InternalFormat = AudioDevice->GetInternalFormat( Wave->NumChannels );		
+		Buffer->InternalFormat = AudioDevice->GetInternalFormat(Wave->NumChannels);
 		Buffer->NumChannels = Wave->NumChannels;
 		Buffer->SampleRate = Wave->SampleRate;
 
@@ -146,20 +163,20 @@ FALSoundBuffer* FALSoundBuffer::CreateNativeBuffer( FALAudioDevice* AudioDevice,
 		{
 			// upload it
 			Buffer->BufferSize = Wave->RawPCMDataSize;
-			alBufferData( Buffer->BufferIds[0], Buffer->InternalFormat, Wave->RawPCMData, Wave->RawPCMDataSize, Buffer->SampleRate );
+			alBufferData(Buffer->BufferIds[0], Buffer->InternalFormat, Wave->RawPCMData, Wave->RawPCMDataSize, Buffer->SampleRate);
 
 			// Free up the data if necessary
-			if( Wave->bDynamicResource )
+			if (Wave->bDynamicResource)
 			{
 				FMemory::Free( Wave->RawPCMData );
-				Wave->RawPCMData = NULL;
+				Wave->RawPCMData = nullptr;
 				Wave->bDynamicResource = false;
 			}
 		}
 		else
 		{
 			// get the raw data
-			uint8* SoundData = ( uint8* )Wave->RawData.Lock( LOCK_READ_ONLY );
+			uint8* SoundData = reinterpret_cast<uint8*>(Wave->RawData.Lock(LOCK_READ_ONLY));
 			// it's (possibly) a pointer to a wave file, so skip over the header
 			int SoundDataSize = Wave->RawData.GetBulkDataSize();
 
@@ -180,16 +197,19 @@ FALSoundBuffer* FALSoundBuffer::CreateNativeBuffer( FALAudioDevice* AudioDevice,
 			Wave->RawData.Unlock();
 		}
 
-		if( AudioDevice->alError( TEXT( "RegisterSound (buffer data)" ) ) || ( Buffer->BufferSize == 0 ) )
+		if (AudioDevice->alError(TEXT( "RegisterSound (buffer data)")) || (Buffer->BufferSize == 0))
 		{
 			Buffer->InternalFormat = 0;
 		}
 
-		if( Buffer->InternalFormat == 0 )
+		if (Buffer->InternalFormat == 0)
 		{
-			UE_LOG ( LogAudio, Log,TEXT( "Audio: sound format not supported for '%s' (%d)" ), *Wave->GetName(), Wave->NumChannels );
+			UE_LOG(LogAudio, Log,TEXT( "Audio: sound format not supported for '%s' (%d)" ), *Wave->GetName(), Wave->NumChannels);
+			UE_LOG(LogALAudio, Warning, TEXT("ALSoundBuffer: sound format not supported for wave '%s'"),
+				*Wave->GetName());
+
 			delete Buffer;
-			Buffer = NULL;
+			Buffer = nullptr;
 		}
 	}
 
diff --git a/Engine/Source/Runtime/Linux/ALAudio/Private/ALAudioSource.cpp b/Engine/Source/Runtime/Linux/ALAudio/Private/ALAudioSource.cpp
--- a/Engine/Source/Runtime/Linux/ALAudio/Private/ALAudioSource.cpp
+++ b/Engine/Source/Runtime/Linux/ALAudio/Private/ALAudioSource.cpp
@@ -12,6 +12,8 @@
  */
 bool FALSoundSource::Init( FWaveInstance* InWaveInstance )
 {
+	check(InWaveInstance);
+
 	if (InWaveInstance->OutputTarget != EAudioOutputTarget::Controller)
 	{
 		// Find matching buffer.
@@ -41,12 +43,15 @@ bool FALSoundSource::Init( FWaveInstance* InWaveInstance )
 			Update();
 
 			// Initialization was successful.
-			return true ;
+			return true;
 		}
 	}
 
 	// Failed to initialize source.
-	return false ;
+	UE_LOG(LogALAudio, Warning, TEXT("Failed to initialize sound source with WaveInstance '%s'."), *InWaveInstance->WaveData->GetName());
+	UE_LOG(LogALAudio, Warning, TEXT("  SampleRate %d"), InWaveInstance->WaveData->SampleRate);
+	UE_LOG(LogALAudio, Warning, TEXT("  Channels %d"), InWaveInstance->WaveData->NumChannels);
+	return false;
 }
 
 /**
diff --git a/Engine/Source/Runtime/Linux/ALAudio/Public/ALAudioDevice.h b/Engine/Source/Runtime/Linux/ALAudio/Public/ALAudioDevice.h
--- a/Engine/Source/Runtime/Linux/ALAudio/Public/ALAudioDevice.h
+++ b/Engine/Source/Runtime/Linux/ALAudio/Public/ALAudioDevice.h
@@ -57,15 +57,14 @@ public:
 	 */
 	static FALSoundBuffer* Init(  FALAudioDevice* ,USoundWave* );
 
-		/**
+	/**
 	 * Static function used to create a IOSAudio buffer and upload decompressed ogg vorbis data to.
 	 *
 	 * @param Wave			USoundWave to use as template and wave source
 	 * @param AudioDevice	audio device to attach created buffer to
 	 * @return FIOSAudioSoundBuffer pointer if buffer creation succeeded, NULL otherwise
 	 */
-	static FALSoundBuffer* CreateNativeBuffer(FALAudioDevice* , USoundWave* );
-
+	static FALSoundBuffer* CreateNativeBuffer(FALAudioDevice* AudioDevice, USoundWave* InWave);
 
 	/**
 	 * Returns the size of this buffer in bytes.
