CIS fix of Android and Linux builds.

diff --git a/Engine/Source/Runtime/Linux/ALAudio/Private/ALAudioBuffer.cpp b/Engine/Source/Runtime/Linux/ALAudio/Private/ALAudioBuffer.cpp
--- a/Engine/Source/Runtime/Linux/ALAudio/Private/ALAudioBuffer.cpp
+++ b/Engine/Source/Runtime/Linux/ALAudio/Private/ALAudioBuffer.cpp
@@ -4,6 +4,7 @@
 #include "Engine.h"
 #include "IAudioFormat.h"
 #include "AudioDecompress.h"
+#include "ContentStreaming.h"
 
 /*------------------------------------------------------------------------------------
 	FALSoundBuffer.
diff --git a/Engine/Source/Runtime/OpenGLDrv/Private/Linux/OpenGLLinux.cpp b/Engine/Source/Runtime/OpenGLDrv/Private/Linux/OpenGLLinux.cpp
--- a/Engine/Source/Runtime/OpenGLDrv/Private/Linux/OpenGLLinux.cpp
+++ b/Engine/Source/Runtime/OpenGLDrv/Private/Linux/OpenGLLinux.cpp
@@ -6,6 +6,7 @@
 
 #include "OpenGLDrvPrivate.h"
 #include "ds_extensions.h"
+#include "ComponentReregisterContext.h"
 
 /*------------------------------------------------------------------------------
 	OpenGL function pointers.
