Fixing Platform Builds

diff --git a/Engine/Source/Runtime/Android/AndroidAudio/Private/AndroidAudioBuffer.cpp b/Engine/Source/Runtime/Android/AndroidAudio/Private/AndroidAudioBuffer.cpp
--- a/Engine/Source/Runtime/Android/AndroidAudio/Private/AndroidAudioBuffer.cpp
+++ b/Engine/Source/Runtime/Android/AndroidAudio/Private/AndroidAudioBuffer.cpp
@@ -16,7 +16,7 @@
  * @param AudioDevice	audio device this sound buffer is going to be attached to.
  */
 FSLESSoundBuffer::FSLESSoundBuffer( FSLESAudioDevice* InAudioDevice )
-:	FAudioBuffer(InAudioDevice),
+:	FSoundBuffer(InAudioDevice),
 	AudioData(NULL),
 	DecompressionState( NULL ),
 	Format(SoundFormat_Invalid)
diff --git a/Engine/Source/Runtime/HTML5/HTML5Audio/Private/HTML5AudioBuffer.cpp b/Engine/Source/Runtime/HTML5/HTML5Audio/Private/HTML5AudioBuffer.cpp
--- a/Engine/Source/Runtime/HTML5/HTML5Audio/Private/HTML5AudioBuffer.cpp
+++ b/Engine/Source/Runtime/HTML5/HTML5Audio/Private/HTML5AudioBuffer.cpp
@@ -112,11 +112,12 @@ FALSoundBuffer* FALSoundBuffer::CreateNativeBuffer( FALAudioDevice* AudioDevice,
 	Wave->InitAudioResource(AudioDevice->GetRuntimeFormat(Wave));
 
 	FALSoundBuffer* Buffer = NULL;
+	FAudioDeviceManager* AudioDeviceManager = GEngine->GetAudioDeviceManager();
 
 	// Find the existing buffer if any
 	if( Wave->ResourceID )
 	{
-		Buffer = static_cast<FALSoundBuffer*>(AudioDevice->WaveBufferMap.FindRef( Wave->ResourceID ));
+		Buffer = static_cast<FALSoundBuffer*>(AudioDeviceManager->WaveBufferMap.FindRef(Wave->ResourceID));
 	}
 
 	if( Buffer == NULL )
diff --git a/Engine/Source/Runtime/HTML5/HTML5Audio/Public/HTML5AudioDevice.h b/Engine/Source/Runtime/HTML5/HTML5Audio/Public/HTML5AudioDevice.h
--- a/Engine/Source/Runtime/HTML5/HTML5Audio/Public/HTML5AudioDevice.h
+++ b/Engine/Source/Runtime/HTML5/HTML5Audio/Public/HTML5AudioDevice.h
@@ -184,6 +184,7 @@ protected:
 class FALAudioDevice : public FAudioDevice
 {
 public: 
+	FALAudioDevice();
 	virtual ~FALAudioDevice() {} 
 
 	virtual FName GetRuntimeFormat(USoundWave* SoundWave) override
diff --git a/Engine/Source/Runtime/Launch/Private/IOS/LaunchIOS.cpp b/Engine/Source/Runtime/Launch/Private/IOS/LaunchIOS.cpp
--- a/Engine/Source/Runtime/Launch/Private/IOS/LaunchIOS.cpp
+++ b/Engine/Source/Runtime/Launch/Private/IOS/LaunchIOS.cpp
@@ -15,17 +15,17 @@ FGameLaunchDaemonMessageHandler GCommandSystem;
 
 void FAppEntry::Suspend()
 {
-	if (GEngine && GEngine->AudioDevice)
+	if (GEngine && GEngine->GetMainAudioDevice())
 	{
-		GEngine->AudioDevice->SuspendContext();
+		GEngine->GetMainAudioDevice()->SuspendContext();
 	}
 }
 
 void FAppEntry::Resume()
 {
-	if (GEngine && GEngine->AudioDevice)
+	if (GEngine && GEngine->GetMainAudioDevice())
 	{
-	    GEngine->AudioDevice->ResumeContext();
+		GEngine->GetMainAudioDevice()->ResumeContext();
 	}
 }
 
diff --git a/Engine/Source/Runtime/Linux/ALAudio/Private/ALAudioBuffer.cpp b/Engine/Source/Runtime/Linux/ALAudio/Private/ALAudioBuffer.cpp
--- a/Engine/Source/Runtime/Linux/ALAudio/Private/ALAudioBuffer.cpp
+++ b/Engine/Source/Runtime/Linux/ALAudio/Private/ALAudioBuffer.cpp
@@ -128,11 +128,14 @@ FALSoundBuffer* FALSoundBuffer::CreateNativeBuffer( FALAudioDevice* AudioDevice,
 	FWaveModInfo WaveInfo;
 	Wave->InitAudioResource(AudioDevice->GetRuntimeFormat(Wave));
 	FALSoundBuffer* Buffer = nullptr;
-	
+
+	FAudioDeviceManager* AudioDeviceManager = GEngine->GetAudioDeviceManager();
+	check(AudioDeviceManager != nullptr);
+
 	// Find the existing buffer if any
 	if (Wave->ResourceID)
 	{
-		Buffer = static_cast<FALSoundBuffer*>(AudioDevice->WaveBufferMap.FindRef( Wave->ResourceID ));
+		Buffer = static_cast<FALSoundBuffer*>(AudioDeviceManager->WaveBufferMap.FindRef(Wave->ResourceID));
 	}
 
 	if (Buffer == nullptr)
@@ -144,8 +147,6 @@ FALSoundBuffer* FALSoundBuffer::CreateNativeBuffer( FALAudioDevice* AudioDevice,
 
 		AudioDevice->alError(TEXT("RegisterSound"));
 
-		FAudioDeviceManager* AudioDeviceManager = GEngine->GetAudioDeviceManager();
-		check(AudioDeviceManager != nullptr);
 
 		AudioDeviceManager->TrackResource(Wave, Buffer);
 
