Fix CIS (wrong initialization list in FALAudioDevice).

diff --git a/Engine/Source/Runtime/Linux/ALAudio/Private/ALAudioDevice.cpp b/Engine/Source/Runtime/Linux/ALAudio/Private/ALAudioDevice.cpp
--- a/Engine/Source/Runtime/Linux/ALAudio/Private/ALAudioDevice.cpp
+++ b/Engine/Source/Runtime/Linux/ALAudio/Private/ALAudioDevice.cpp
@@ -34,7 +34,7 @@ IMPLEMENT_MODULE(FALAudioDeviceModule, ALAudio );
 	UALAudioDevice constructor and UObject interface.
 ------------------------------------------------------------------------------------*/
 FALAudioDevice::FALAudioDevice()
-	, HardwareDevice(nullptr)
+	: HardwareDevice(nullptr)
 	, SoundContext(nullptr)
 {}
 
