Linux: support switching between multiple devices in ALAudio.

- Fixes audio device brokennes after PIE (UE-12236)

diff --git a/Engine/Source/Runtime/Linux/ALAudio/Private/ALAudioBuffer.cpp b/Engine/Source/Runtime/Linux/ALAudio/Private/ALAudioBuffer.cpp
--- a/Engine/Source/Runtime/Linux/ALAudio/Private/ALAudioBuffer.cpp
+++ b/Engine/Source/Runtime/Linux/ALAudio/Private/ALAudioBuffer.cpp
@@ -28,7 +28,7 @@ FALSoundBuffer::FALSoundBuffer( FALAudioDevice* InAudioDevice )
 FALSoundBuffer::~FALSoundBuffer( void )
 {
 	// Delete AL buffers.
-	alDeleteBuffers( 1, BufferIds );
+	alDeleteBuffers(1, &BufferId);
 }
 
 /**
@@ -143,11 +143,9 @@ FALSoundBuffer* FALSoundBuffer::CreateNativeBuffer( FALAudioDevice* AudioDevice,
 		// Create new buffer.
 		Buffer = new FALSoundBuffer(AudioDevice);
 
-		alGenBuffers(1, Buffer->BufferIds);
-
+		alGenBuffers(1, &Buffer->BufferId);
 		AudioDevice->alError(TEXT("RegisterSound"));
 
-
 		AudioDeviceManager->TrackResource(Wave, Buffer);
 
 		Buffer->InternalFormat = AudioDevice->GetInternalFormat(Wave->NumChannels);
@@ -158,7 +156,7 @@ FALSoundBuffer* FALSoundBuffer::CreateNativeBuffer( FALAudioDevice* AudioDevice,
 		{
 			// upload it
 			Buffer->BufferSize = Wave->RawPCMDataSize;
-			alBufferData(Buffer->BufferIds[0], Buffer->InternalFormat, Wave->RawPCMData, Wave->RawPCMDataSize, Buffer->SampleRate);
+			alBufferData(Buffer->BufferId, Buffer->InternalFormat, Wave->RawPCMData, Wave->RawPCMDataSize, Buffer->SampleRate);
 
 			// Free up the data if necessary
 			if (Wave->bDynamicResource)
@@ -187,7 +185,7 @@ FALSoundBuffer* FALSoundBuffer::CreateNativeBuffer( FALAudioDevice* AudioDevice,
 			Buffer->BufferSize = SoundDataSize;
 
 			// upload it
-			alBufferData( Buffer->BufferIds[0], Buffer->InternalFormat, SoundData, Buffer->BufferSize, Buffer->SampleRate );
+			alBufferData( Buffer->BufferId, Buffer->InternalFormat, SoundData, Buffer->BufferSize, Buffer->SampleRate );
 			// unload it
 			Wave->RawData.Unlock();
 		}
diff --git a/Engine/Source/Runtime/Linux/ALAudio/Private/ALAudioDevice.cpp b/Engine/Source/Runtime/Linux/ALAudio/Private/ALAudioDevice.cpp
--- a/Engine/Source/Runtime/Linux/ALAudio/Private/ALAudioDevice.cpp
+++ b/Engine/Source/Runtime/Linux/ALAudio/Private/ALAudioDevice.cpp
@@ -60,7 +60,7 @@ void FALAudioDevice::TeardownHardware( void )
 	// Disable the context
 	if( &alcMakeContextCurrent )
 	{	
-		alcMakeContextCurrent( NULL );
+		alcMakeContextCurrent(nullptr);
 	}
 
 	// Destroy the context
@@ -72,9 +72,13 @@ void FALAudioDevice::TeardownHardware( void )
 
 	// Close the hardware device
 	if( &alcCloseDevice )
-	{	
+	{
+		checkf(HardwareDevice, TEXT("Tearing down invalid OpenAL device! (HardwareDevice should not be null)."));
+		const ALCchar* DeviceName = alcGetString(HardwareDevice, ALC_DEVICE_SPECIFIER);
+		UE_LOG(LogALAudio, Log, TEXT("Closing ALAudio device : %s"), StringCast<TCHAR>(static_cast<const ANSICHAR*>(DeviceName)).Get());
+
 		alcCloseDevice( HardwareDevice );
-		HardwareDevice = NULL;
+		HardwareDevice = nullptr;
 	}
 
 }
@@ -103,9 +107,8 @@ bool FALAudioDevice::InitializeHardware( void )
 	// Load ogg and vorbis dlls if they haven't been loaded yet
 	//LoadVorbisLibraries();
 
-	
 	// Open device
-	HardwareDevice = alcOpenDevice(  NULL );
+	HardwareDevice = alcOpenDevice(nullptr);
 	if( !HardwareDevice )
 	{
 		UE_LOG(LogALAudio, Log, TEXT( "ALAudio: no OpenAL devices found." ) );
@@ -134,8 +137,7 @@ bool FALAudioDevice::InitializeHardware( void )
 		return false ;
 	}
 
-	alcMakeContextCurrent( SoundContext );
-	
+	alcMakeContextCurrent(SoundContext);
 
 	// Make sure everything happened correctly
 	if( alError( TEXT( "Init" ) ) )
@@ -206,7 +208,6 @@ void FALAudioDevice::Update( bool Realtime )
 	FVector ListenerFront		= Listeners[ 0 ].GetFront();
 	FVector ListenerUp			= Listeners[ 0 ].GetUp();
 
-
 	// Set Player position
 	FVector Location;
 
@@ -223,7 +224,7 @@ void FALAudioDevice::Update( bool Realtime )
 	Orientation[0].X = ListenerFront.X;
 	Orientation[0].Y = ListenerFront.Z; // Z/Y swapped on purpose, see file header	
 	Orientation[0].Z = ListenerFront.Y; // Z/Y swapped on purpose, see file header
-	
+
 	// See file header for coordinate system explanation.
 	Orientation[1].X = ListenerUp.X;
 	Orientation[1].Y = ListenerUp.Z; // Z/Y swapped on purpose, see file header
@@ -232,7 +233,9 @@ void FALAudioDevice::Update( bool Realtime )
 	// Make the listener still and the sounds move relatively -- this allows 
 	// us to scale the doppler effect on a per-sound basis.
 	FVector Velocity = FVector( 0.0f, 0.0f, 0.0f );
-	
+
+	MakeCurrent(TEXT("FALAudioDevice::Update()"));
+
 	alListenerfv( AL_POSITION, ( ALfloat* )&Location );
 	alListenerfv( AL_ORIENTATION, ( ALfloat* )&Orientation[0] );
 	alListenerfv( AL_VELOCITY, ( ALfloat* )&Velocity );
@@ -275,6 +278,15 @@ ALuint FALAudioDevice::GetInternalFormat( int NumChannels )
 	return( InternalFormat );
 }
 
+void FALAudioDevice::MakeCurrent(const TCHAR * CallSiteIdentifier)
+{
+	checkf(SoundContext, TEXT("Unitialiized sound context in FALAudioDevice::MakeCurrent()!"));
+	if (!alcMakeContextCurrent(SoundContext))
+	{
+		alError(CallSiteIdentifier ? CallSiteIdentifier : TEXT("FALAudioDevice::MakeCurrent()"));
+	}
+}
+
 /*------------------------------------------------------------------------------------
 OpenAL utility functions
 ------------------------------------------------------------------------------------*/
diff --git a/Engine/Source/Runtime/Linux/ALAudio/Private/ALAudioSource.cpp b/Engine/Source/Runtime/Linux/ALAudio/Private/ALAudioSource.cpp
--- a/Engine/Source/Runtime/Linux/ALAudio/Private/ALAudioSource.cpp
+++ b/Engine/Source/Runtime/Linux/ALAudio/Private/ALAudioSource.cpp
@@ -15,6 +15,7 @@
 bool FALSoundSource::Init( FWaveInstance* InWaveInstance )
 {
 	check(InWaveInstance);
+	GetALDevice()->MakeCurrent(TEXT("FALSoundSource::Init()"));
 
 	if (InWaveInstance->OutputTarget != EAudioOutputTarget::Controller)
 	{
@@ -33,13 +34,13 @@ bool FALSoundSource::Init( FWaveInstance* InWaveInstance )
 			alSourcei( SourceId, AL_LOOPING, ( WaveInstance->LoopingMode == LOOP_Forever ) ? AL_TRUE : AL_FALSE );
 
 			// Always queue up the first buffer
-			alSourceQueueBuffers( SourceId, 1, Buffer->BufferIds );	
+			alSourceQueueBuffers(SourceId, 1, &Buffer->BufferId);
 			if( WaveInstance->LoopingMode == LOOP_WithNotification )
 			{		
 				// We queue the sound twice for wave instances that use seamless looping so we can have smooth 
 				// loop transitions. The downside is that we might play at most one frame worth of audio from the 
 				// beginning of the wave after the wave stops looping.
-				alSourceQueueBuffers( SourceId, 1, Buffer->BufferIds );
+				alSourceQueueBuffers(SourceId, 1, &Buffer->BufferId);
 			}
 
 			Update();
@@ -61,6 +62,8 @@ bool FALSoundSource::Init( FWaveInstance* InWaveInstance )
  */
 FALSoundSource::~FALSoundSource( void )
 {
+	GetALDevice()->MakeCurrent(TEXT("FALSoundSource::~FALSoundSource()"));
+
 	// @todo openal: What do we do here
 	/// AudioDevice->DestroyEffect( this );
 
@@ -80,6 +83,8 @@ void FALSoundSource::Update( void )
 		return;
 	}
 
+	GetALDevice()->MakeCurrent(TEXT("FALSoundSource::Update()"));
+
 	float Volume = 1.0f;
 
 	if (AudioDevice->bIsDeviceMuted)
@@ -146,6 +151,8 @@ void FALSoundSource::Play( void )
 {
 	if( WaveInstance )
 	{
+		GetALDevice()->MakeCurrent(TEXT("FALSoundSource::Play()"));
+
 		alSourcePlay( SourceId );
 		Paused = false;
 		Playing = true;
@@ -159,6 +166,8 @@ void FALSoundSource::Stop( void )
 {
 	if( WaveInstance )
 	{
+		GetALDevice()->MakeCurrent(TEXT("FALSoundSource::Stop()"));
+
 		alSourceStop( SourceId );
 		// This clears out any pending buffers that may or may not be queued or played
 		alSourcei( SourceId, AL_BUFFER, 0 );
@@ -177,6 +186,8 @@ void FALSoundSource::Pause( void )
 {
 	if( WaveInstance )
 	{
+		GetALDevice()->MakeCurrent(TEXT("FALSoundSource::Pause()"));
+
 		alSourcePause( SourceId );
 		Paused = true;
 	}
@@ -189,14 +200,16 @@ bool FALSoundSource::IsSourceFinished( void )
 {
 	ALint State = AL_STOPPED;
 
+	GetALDevice()->MakeCurrent(TEXT("FALSoundSource::IsSourceFinished()"));
+
 	// Check the source for data to continue playing
 	alGetSourcei( SourceId, AL_SOURCE_STATE, &State );
 	if( State == AL_PLAYING || State == AL_PAUSED )
 	{
-		return( false );
+		return false;
 	}
 
-	return( true );
+	return true;
 }
 
 /** 
@@ -204,6 +217,8 @@ bool FALSoundSource::IsSourceFinished( void )
  */
 void FALSoundSource::HandleQueuedBuffer( void )
 {
+	GetALDevice()->MakeCurrent(TEXT("FALSoundSource::HandleQueuedBuffer()"));
+
 	ALuint	DequeuedBuffer;
 
 	// Unqueue the processed buffers
@@ -213,7 +228,7 @@ void FALSoundSource::HandleQueuedBuffer( void )
 	WaveInstance->NotifyFinished();
 
 	// Queue the same packet again for looping
-	alSourceQueueBuffers( SourceId, 1, Buffer->BufferIds );
+	alSourceQueueBuffers(SourceId, 1, &Buffer->BufferId);
 }
 
 /**
@@ -224,6 +239,8 @@ void FALSoundSource::HandleQueuedBuffer( void )
  */
 bool FALSoundSource::IsFinished( void )
 {
+	GetALDevice()->MakeCurrent(TEXT("FALSoundSource::IsFinished()"));
+
 	if( WaveInstance )
 	{
 		// Check for a non starved, stopped source
@@ -261,8 +278,8 @@ bool FALSoundSource::IsFinished( void )
 			}
 		}
 
-		return( false  );
+		return false;
 	}
 
-	return(  true );
+	return true;
 }
diff --git a/Engine/Source/Runtime/Linux/ALAudio/Public/ALAudioDevice.h b/Engine/Source/Runtime/Linux/ALAudio/Public/ALAudioDevice.h
--- a/Engine/Source/Runtime/Linux/ALAudio/Public/ALAudioDevice.h
+++ b/Engine/Source/Runtime/Linux/ALAudio/Public/ALAudioDevice.h
@@ -83,8 +83,8 @@ public:
 
 	/** Audio device this buffer is attached to */
 	FALAudioDevice*			AudioDevice;
-	/** Array of buffer ids used to reference the data stored in AL. */
-	ALuint					BufferIds[2];
+	/** Buffer id used to reference the data stored in AL. */
+	ALuint					BufferId;
 	/** Format of the data internal to OpenAL */
 	ALuint					InternalFormat;
 
@@ -175,6 +175,9 @@ protected:
 	FALSoundBuffer*		Buffer;
 
 	friend class FALAudioDevice;
+
+	/** Convenience function to avoid casting in multiple places */
+	FORCEINLINE FALAudioDevice * GetALDevice();
 };
 
 /**
@@ -219,9 +222,17 @@ public:
 	// Error checking.
 	bool alError( const TCHAR* Text, bool Log = true );
 
+	/** 
+	 * Makes sure correct context is set.
+	 * 
+	 * @param CallSiteIdentifier identifies the call site for debugging/error reporting purposes
+	 */
+	void MakeCurrent(const TCHAR * CallSiteIdentifier = nullptr);
+
 protected:
 
 	virtual FSoundSource* CreateSoundSource() override; 
+
 	/** Returns the enum for the internal format for playing a sound with this number of channels. */
 	ALuint GetInternalFormat( int NumChannels );
 
@@ -238,6 +249,7 @@ protected:
 	ALCdevice*									HardwareDevice;
 	ALCcontext*									SoundContext;
 	void*										DLLHandle;
+
 	/** Formats for multichannel sounds */
 	ALenum										Surround40Format;
 	ALenum										Surround51Format;
@@ -247,3 +259,7 @@ protected:
 	friend class FALSoundBuffer;
 };
 
+FORCEINLINE FALAudioDevice * FALSoundSource::GetALDevice()
+{
+	return static_cast<FALAudioDevice *>(AudioDevice);
+}
