Linux: fix for sounds playing in the editor only once.

- PR #1166 contributed by [REDACTED_PERSON].
- Also a lot of whitespace cleanup.

diff --git a/Engine/Source/Runtime/Linux/ALAudio/Private/ALAudioBuffer.cpp b/Engine/Source/Runtime/Linux/ALAudio/Private/ALAudioBuffer.cpp
--- a/Engine/Source/Runtime/Linux/ALAudio/Private/ALAudioBuffer.cpp
+++ b/Engine/Source/Runtime/Linux/ALAudio/Private/ALAudioBuffer.cpp
@@ -10,19 +10,20 @@
 /*------------------------------------------------------------------------------------
 	FALSoundBuffer.
 ------------------------------------------------------------------------------------*/
-/** 
+/**
  * Constructor
  *
  * @param AudioDevice	audio device this sound buffer is going to be attached to.
  */
 FALSoundBuffer::FALSoundBuffer( FALAudioDevice* InAudioDevice )
-	: FSoundBuffer(InAudioDevice)
+	: FSoundBuffer(InAudioDevice),
+	  BufferId(0)
 {
 }
 
 /**
- * Destructor 
- * 
+ * Destructor
+ *
  * Frees wave data and detaches itself from audio device.
  */
 FALSoundBuffer::~FALSoundBuffer( void )
@@ -68,9 +69,9 @@ FALSoundBuffer* FALSoundBuffer::Init(FALAudioDevice* AudioDevice, USoundWave* In
 			Buffer = static_cast<FALSoundBuffer*>(AudioDeviceManager->WaveBufferMap.FindRef(InWave->ResourceID));
 		}
 
-		if (!Buffer)
+		if (!Buffer || !Buffer->BufferId)
 		{
-			Buffer = CreateNativeBuffer(AudioDevice, InWave);
+			CreateNativeBuffer(AudioDevice, InWave, Buffer);
 		}
 		break;
 
@@ -84,7 +85,7 @@ FALSoundBuffer* FALSoundBuffer::Init(FALAudioDevice* AudioDevice, USoundWave* In
 	case DTYPE_RealTime:
 	default:
 		// Invalid will be set if the wave cannot be played
-		UE_LOG(LogALAudio, Warning, TEXT("ALSoundBuffer wave '%s' has an invalid decompression type %d."), 
+		UE_LOG(LogALAudio, Warning, TEXT("ALSoundBuffer wave '%s' has an invalid decompression type %d."),
 			*InWave->GetName(), static_cast<EDecompressionType>(InWave->DecompressionType));
 		break;
 	}
@@ -98,112 +99,95 @@ FALSoundBuffer* FALSoundBuffer::Init(FALAudioDevice* AudioDevice, USoundWave* In
 	return Buffer;
 }
 
-FALSoundBuffer* FALSoundBuffer::CreateNativeBuffer( FALAudioDevice* AudioDevice, USoundWave* Wave)
+void FALSoundBuffer::CreateNativeBuffer(FALAudioDevice* AudioDevice, USoundWave* Wave, FALSoundBuffer* &Buffer)
 {
-	SCOPE_CYCLE_COUNTER( STAT_AudioResourceCreationTime );
-
-	// Check to see if thread has finished decompressing on the other thread
-	if (Wave->AudioDecompressor)
+	// Check if Buffer already exists
+	if (Buffer)
 	{
-		Wave->AudioDecompressor->EnsureCompletion();
-
-		// Remove the decompressor
-		delete Wave->AudioDecompressor;
-		Wave->AudioDecompressor = nullptr;
+		// Assign the new AudioDevice
+		Buffer->AudioDevice = AudioDevice;
 	}
-
-	// Can't create a buffer without any source data
-	if (Wave == nullptr)
+	else
 	{
-		return nullptr;
-	}
+		SCOPE_CYCLE_COUNTER( STAT_AudioResourceCreationTime );
 
-	if (Wave->NumChannels == 0)
-	{
-		UE_LOG(LogALAudio, Warning, TEXT("ALSoundBuffer: cannot create buffer for wave '%s' with 0 channels"),
-			*Wave->GetName());
-		return nullptr;
-	}
-
-	Wave->InitAudioResource(AudioDevice->GetRuntimeFormat(Wave));
-	FALSoundBuffer* Buffer = nullptr;
+		// Check to see if thread has finished decompressing on the other thread
+		if (Wave->AudioDecompressor)
+		{
+			Wave->AudioDecompressor->EnsureCompletion();
 
-	FAudioDeviceManager* AudioDeviceManager = GEngine->GetAudioDeviceManager();
-	check(AudioDeviceManager != nullptr);
+			// Remove the decompressor
+			delete Wave->AudioDecompressor;
+			Wave->AudioDecompressor = nullptr;
+		}
 
-	// Find the existing buffer if any
-	if (Wave->ResourceID)
-	{
-		Buffer = static_cast<FALSoundBuffer*>(AudioDeviceManager->WaveBufferMap.FindRef(Wave->ResourceID));
-	}
+		Wave->InitAudioResource(AudioDevice->GetRuntimeFormat(Wave));
 
-	if (Buffer == nullptr)
-	{
 		// Create new buffer.
 		Buffer = new FALSoundBuffer(AudioDevice);
 
-		alGenBuffers(1, &Buffer->BufferId);
-		AudioDevice->alError(TEXT("RegisterSound"));
-
-		AudioDeviceManager->TrackResource(Wave, Buffer);
-
 		Buffer->InternalFormat = AudioDevice->GetInternalFormat(Wave->NumChannels);
 		Buffer->NumChannels = Wave->NumChannels;
 		Buffer->SampleRate = Wave->SampleRate;
 
-		if (Wave->RawPCMData)
-		{
-			// upload it
-			Buffer->BufferSize = Wave->RawPCMDataSize;
-			alBufferData(Buffer->BufferId, Buffer->InternalFormat, Wave->RawPCMData, Wave->RawPCMDataSize, Buffer->SampleRate);
-
-			// Free up the data if necessary
-			if (Wave->bDynamicResource)
-			{
-				FMemory::Free( Wave->RawPCMData );
-				Wave->RawPCMData = nullptr;
-				Wave->bDynamicResource = false;
-			}
-		}
-		else
+		FAudioDeviceManager* AudioDeviceManager = GEngine->GetAudioDeviceManager();
+		check(AudioDeviceManager != nullptr);
+		AudioDeviceManager->TrackResource(Wave, Buffer);
+	}
+
+	// Generate the new OpenAL buffer
+	alGenBuffers(1, &Buffer->BufferId);
+	AudioDevice->alError(TEXT("RegisterSound"));
+
+	if (Wave->RawPCMData)
+	{
+		// upload it
+		Buffer->BufferSize = Wave->RawPCMDataSize;
+		alBufferData(Buffer->BufferId, Buffer->InternalFormat, Wave->RawPCMData, Wave->RawPCMDataSize, Buffer->SampleRate);
+
+		// Free up the data if necessary
+		if (Wave->bDynamicResource)
 		{
-			// get the raw data
-			uint8* SoundData = reinterpret_cast<uint8*>(Wave->RawData.Lock(LOCK_READ_ONLY));
-			// it's (possibly) a pointer to a wave file, so skip over the header
-			int SoundDataSize = Wave->RawData.GetBulkDataSize();
-
-			// is there a wave header?
-			FWaveModInfo WaveInfo;
-			if (WaveInfo.ReadWaveInfo(SoundData, SoundDataSize))
-			{
-				// if so, modify the location and size of the sound data based on header
-				SoundData = WaveInfo.SampleDataStart;
-				SoundDataSize = WaveInfo.SampleDataSize;
-			}
-			// let the Buffer know the final size
-			Buffer->BufferSize = SoundDataSize;
-
-			// upload it
-			alBufferData( Buffer->BufferId, Buffer->InternalFormat, SoundData, Buffer->BufferSize, Buffer->SampleRate );
-			// unload it
-			Wave->RawData.Unlock();
+			FMemory::Free( Wave->RawPCMData );
+			Wave->RawPCMData = nullptr;
+			Wave->bDynamicResource = false;
 		}
-
-		if (AudioDevice->alError(TEXT( "RegisterSound (buffer data)")) || (Buffer->BufferSize == 0))
+	}
+	else
+	{
+		// get the raw data
+		uint8* SoundData = reinterpret_cast<uint8*>(Wave->RawData.Lock(LOCK_READ_ONLY));
+		// it's (possibly) a pointer to a wave file, so skip over the header
+		int SoundDataSize = Wave->RawData.GetBulkDataSize();
+
+		// is there a wave header?
+		FWaveModInfo WaveInfo;
+		if (WaveInfo.ReadWaveInfo(SoundData, SoundDataSize))
 		{
-			Buffer->InternalFormat = 0;
+			// if so, modify the location and size of the sound data based on header
+			SoundData = WaveInfo.SampleDataStart;
+			SoundDataSize = WaveInfo.SampleDataSize;
 		}
+		// let the Buffer know the final size
+		Buffer->BufferSize = SoundDataSize;
 
-		if (Buffer->InternalFormat == 0)
-		{
-			UE_LOG(LogAudio, Log,TEXT( "Audio: sound format not supported for '%s' (%d)" ), *Wave->GetName(), Wave->NumChannels);
-			UE_LOG(LogALAudio, Warning, TEXT("ALSoundBuffer: sound format not supported for wave '%s'"),
-				*Wave->GetName());
+		// upload it
+		alBufferData( Buffer->BufferId, Buffer->InternalFormat, SoundData, Buffer->BufferSize, Buffer->SampleRate );
+		// unload it
+		Wave->RawData.Unlock();
+	}
 
-			delete Buffer;
-			Buffer = nullptr;
-		}
+	if (AudioDevice->alError(TEXT( "RegisterSound (buffer data)")) || (Buffer->BufferSize == 0))
+	{
+		Buffer->InternalFormat = 0;
 	}
 
-	return Buffer;
+	if (Buffer->InternalFormat == 0)
+	{
+		UE_LOG(LogAudio, Log,TEXT( "Audio: sound format not supported for '%s' (%d)" ), *Wave->GetName(), Wave->NumChannels);
+		UE_LOG(LogALAudio, Warning, TEXT("ALSoundBuffer: sound format not supported for wave '%s'"), *Wave->GetName());
+
+		delete Buffer;
+		Buffer = nullptr;
+	}
 }
diff --git a/Engine/Source/Runtime/Linux/ALAudio/Private/ALAudioDevice.cpp b/Engine/Source/Runtime/Linux/ALAudio/Private/ALAudioDevice.cpp
--- a/Engine/Source/Runtime/Linux/ALAudio/Private/ALAudioDevice.cpp
+++ b/Engine/Source/Runtime/Linux/ALAudio/Private/ALAudioDevice.cpp
@@ -40,32 +40,43 @@ FALAudioDevice::FALAudioDevice()
 
 void FALAudioDevice::TeardownHardware( void )
 {
-	// Flush stops all sources and deletes all buffers so sources can be safely deleted below.
-	Flush( NULL );
-	
 	// Push any pending data to the hardware
 	if( &alcProcessContext )
-	{	
+	{
 		alcProcessContext( SoundContext );
 	}
 
 	// Destroy all sound sources
-	for( int i = 0; i < Sources.Num(); i++ )
+	for( int32 i = 0; i < Sources.Num(); i++ )
 	{
 		delete Sources[ i ];
 	}
 	Sources.Empty();
 	FreeSources.Empty();
 
+	// Destroy OpenAL buffers associated with this audio device
+	FAudioDeviceManager* AudioDeviceManager = GEngine->GetAudioDeviceManager();
+	check(AudioDeviceManager != nullptr);
+	for( int32 i = 0; i < AudioDeviceManager->Buffers.Num(); i++ )
+	{
+		FALSoundBuffer* Buffer = static_cast<FALSoundBuffer*>(AudioDeviceManager->Buffers[i]);
+		if( Buffer->AudioDevice == this )
+		{
+			alDeleteBuffers(1, &Buffer->BufferId);
+			Buffer->AudioDevice = nullptr;
+			Buffer->BufferId = 0;
+		}
+	}
+
 	// Disable the context
 	if( &alcMakeContextCurrent )
-	{	
+	{
 		alcMakeContextCurrent(nullptr);
 	}
 
 	// Destroy the context
 	if( &alcDestroyContext )
-	{	
+	{
 		alcDestroyContext( SoundContext );
 		SoundContext = NULL;
 	}
@@ -89,7 +100,7 @@ void FALAudioDevice::TeardownHardware( void )
 /**
  * Initializes the audio device and creates sources.
  *
- * @warning: 
+ * @warning:
  *
  * @return TRUE if initialization was successful, FALSE otherwise
  */
@@ -120,12 +131,12 @@ bool FALAudioDevice::InitializeHardware( void )
 	UE_LOG(LogALAudio, Log, TEXT("ALAudio device opened : %s"), StringCast<TCHAR>(static_cast<const ANSICHAR*>(OpenedDeviceName)).Get());
 
 	// Create a context
-	int Caps[] = 
-	{ 
-		ALC_FREQUENCY, 44100, 
-		ALC_STEREO_SOURCES, 4, 
-		0, 0 
-    };
+	int Caps[] =
+	{
+		ALC_FREQUENCY, 44100,
+		ALC_STEREO_SOURCES, 4,
+		0, 0
+	};
 #if PLATFORM_HTML5_WIN32 || PLATFORM_LINUX
 	SoundContext = alcCreateContext( HardwareDevice, Caps );
 #elif PLATFORM_HTML5
@@ -150,7 +161,7 @@ bool FALAudioDevice::InitializeHardware( void )
 	UE_LOG(LogALAudio, Log, TEXT("AL_RENDERER    : %s"), StringCast<TCHAR>(static_cast<const ANSICHAR*>(alGetString(AL_RENDERER))).Get());
 	UE_LOG(LogALAudio, Log, TEXT("AL_VERSION     : %s"), StringCast<TCHAR>(static_cast<const ANSICHAR*>(alGetString(AL_VERSION))).Get());
 	UE_LOG(LogALAudio, Log, TEXT("AL_EXTENSIONS  : %s"), StringCast<TCHAR>(static_cast<const ANSICHAR*>(alGetString(AL_EXTENSIONS))).Get());
- 
+
 	// Get the enums for multichannel support
 #if !PLATFORM_HTML5
 	Surround40Format = alGetEnumValue( "AL_FORMAT_QUAD16" );
@@ -190,7 +201,7 @@ bool FALAudioDevice::InitializeHardware( void )
 	// Use our own distance model.
 	alDistanceModel( AL_NONE );
 
-	// Set up a default (nop) effects manager 
+	// Set up a default (nop) effects manager
 	Effects = new FAudioEffectsManager( this );
 
 	return true ;
@@ -216,13 +227,13 @@ void FALAudioDevice::Update( bool Realtime )
 	Location.Y = ListenerLocation.Z; // Z/Y swapped on purpose, see file header
 	Location.Z = ListenerLocation.Y; // Z/Y swapped on purpose, see file header
 	Location *= AUDIO_DISTANCE_FACTOR;
-	
+
 	// Set Player orientation.
 	FVector Orientation[2];
 
 	// See file header for coordinate system explanation.
 	Orientation[0].X = ListenerFront.X;
-	Orientation[0].Y = ListenerFront.Z; // Z/Y swapped on purpose, see file header	
+	Orientation[0].Y = ListenerFront.Z; // Z/Y swapped on purpose, see file header
 	Orientation[0].Z = ListenerFront.Y; // Z/Y swapped on purpose, see file header
 
 	// See file header for coordinate system explanation.
@@ -230,7 +241,7 @@ void FALAudioDevice::Update( bool Realtime )
 	Orientation[1].Y = ListenerUp.Z; // Z/Y swapped on purpose, see file header
 	Orientation[1].Z = ListenerUp.Y; // Z/Y swapped on purpose, see file header
 
-	// Make the listener still and the sounds move relatively -- this allows 
+	// Make the listener still and the sounds move relatively -- this allows
 	// us to scale the doppler effect on a per-sound basis.
 	FVector Velocity = FVector( 0.0f, 0.0f, 0.0f );
 
@@ -241,7 +252,7 @@ void FALAudioDevice::Update( bool Realtime )
 	alListenerfv( AL_VELOCITY, ( ALfloat* )&Velocity );
 
 	alError( TEXT( "UALAudioDevice::Update" ) );
-}	
+}
 
 ALuint FALAudioDevice::GetInternalFormat( int NumChannels )
 {
@@ -265,14 +276,14 @@ ALuint FALAudioDevice::GetInternalFormat( int NumChannels )
 		break;
 	case 6:
 		InternalFormat = Surround51Format;
-		break;	
+		break;
 	case 7:
 		InternalFormat = Surround61Format;
-		break;	
+		break;
 	case 8:
 		InternalFormat = Surround71Format;
 		break;
-#endif 
+#endif
 	}
 
 	return( InternalFormat );
@@ -290,29 +301,6 @@ void FALAudioDevice::MakeCurrent(const TCHAR * CallSiteIdentifier)
 /*------------------------------------------------------------------------------------
 OpenAL utility functions
 ------------------------------------------------------------------------------------*/
-//
-//	FindExt
-//
-bool FALAudioDevice::FindExt( const TCHAR* Name )
-{
-	return alIsExtensionPresent(TCHAR_TO_ANSI(Name)) || alcIsExtensionPresent(HardwareDevice, TCHAR_TO_ANSI(Name));
-}
-
-//
-//	FindProc
-//
-void FALAudioDevice::FindProc( void*& ProcAddress, char* Name, char* SupportName, bool& Supports, bool AllowExt )
-{
-	
-}
-
-//
-//	FindProcs
-//
-void FALAudioDevice::FindProcs( bool AllowExt )
-{
-
-}
 
 //
 //	alError
@@ -322,8 +310,8 @@ bool FALAudioDevice::alError( const TCHAR* Text, bool Log )
 	ALenum Error = alGetError();
 	if( Error != AL_NO_ERROR )
 	{
-		do 
-		{		
+		do
+		{
 			if( Log )
 			{
 				switch ( Error )
@@ -351,10 +339,10 @@ bool FALAudioDevice::alError( const TCHAR* Text, bool Log )
 		}
 		while( ( Error = alGetError() ) != AL_NO_ERROR );
 
-		return true ; 
+		return true;
 	}
 
-	return false ;
+	return false;
 }
 
 FSoundSource* FALAudioDevice::CreateSoundSource()
diff --git a/Engine/Source/Runtime/Linux/ALAudio/Public/ALAudioDevice.h b/Engine/Source/Runtime/Linux/ALAudio/Public/ALAudioDevice.h
--- a/Engine/Source/Runtime/Linux/ALAudio/Public/ALAudioDevice.h
+++ b/Engine/Source/Runtime/Linux/ALAudio/Public/ALAudioDevice.h
@@ -4,7 +4,7 @@
 	ALAudioDevice.h: Unreal OpenAL audio interface object.
 =============================================================================*/
 
-#pragma  once 
+#pragma  once
 
 /*------------------------------------------------------------------------------------
 	Dependencies, helpers & forward declarations.
@@ -25,7 +25,7 @@ class FALAudioDevice;
 #include "AL/alc.h"
 #if SUPPORTS_PRAGMA_PACK
 #pragma pack (pop)
-#endif 
+#endif
 
 DECLARE_LOG_CATEGORY_EXTERN(LogALAudio, Log, All);
 
@@ -36,21 +36,21 @@ DECLARE_LOG_CATEGORY_EXTERN(LogALAudio, Log, All);
 /**
  * OpenAL implementation of FSoundBuffer, containing the wave data and format information.
  */
-class FALSoundBuffer : public FSoundBuffer 
+class FALSoundBuffer : public FSoundBuffer
 {
 public:
-	/** 
+	/**
 	 * Constructor
 	 *
 	 * @param AudioDevice	audio device this sound buffer is going to be attached to.
 	 */
 	FALSoundBuffer( FALAudioDevice* AudioDevice );
-	
+
 	/**
-	 * Destructor 
-	 * 
+	 * Destructor
+	 *
 	 * Frees wave data and detaches itself from audio device.
- 	 */
+	 */
 	~FALSoundBuffer( void );
 
 	/**
@@ -69,20 +69,18 @@ public:
 	 * @param AudioDevice	audio device to attach created buffer to
 	 * @return FIOSAudioSoundBuffer pointer if buffer creation succeeded, NULL otherwise
 	 */
-	static FALSoundBuffer* CreateNativeBuffer(FALAudioDevice* AudioDevice, USoundWave* InWave);
+	static void CreateNativeBuffer(FALAudioDevice* AudioDevice, USoundWave* InWave, FALSoundBuffer* &Buffer);
 
 	/**
 	 * Returns the size of this buffer in bytes.
 	 *
 	 * @return Size in bytes
 	 */
-	int GetSize( void ) 
-	{ 
-		return( BufferSize ); 
+	int GetSize( void )
+	{
+		return( BufferSize );
 	}
 
-	/** Audio device this buffer is attached to */
-	FALAudioDevice*			AudioDevice;
 	/** Buffer id used to reference the data stored in AL. */
 	ALuint					BufferId;
 	/** Format of the data internal to OpenAL */
@@ -111,7 +109,7 @@ public:
 		Buffer( NULL )
 	{}
 
-	/** 
+	/**
 	 * Destructor
 	 */
 	~FALSoundSource( void );
@@ -126,17 +124,17 @@ public:
 
 	/**
 	 * Updates the source specific parameter like e.g. volume and pitch based on the associated
-	 * wave instance.	
+	 * wave instance.
 	 */
 	virtual void Update( void );
 
 	/**
-	 * Plays the current wave instance.	
+	 * Plays the current wave instance.
 	 */
 	virtual void Play( void );
 
 	/**
-	 * Stops the current wave instance and detaches it from the source.	
+	 * Stops the current wave instance and detaches it from the source.
 	 */
 	virtual void Stop( void );
 
@@ -148,22 +146,22 @@ public:
 	/**
 	 * Queries the status of the currently associated wave instance.
 	 *
-	 * @return	TRUE if the wave instance/ source has finished playback and FALSE if it is 
+	 * @return	TRUE if the wave instance/ source has finished playback and FALSE if it is
 	 *			currently playing or paused.
 	 */
 	virtual bool IsFinished( void );
 
-	/** 
+	/**
 	 * Access function for the source id
 	 */
 	ALuint GetSourceId( void ) const { return( SourceId ); }
 
-	/** 
+	/**
 	 * Returns TRUE if an OpenAL source has finished playing
 	 */
 	bool IsSourceFinished( void );
 
-	/** 
+	/**
 	 * Handle dequeuing and requeuing of a single buffer
 	 */
 	void HandleQueuedBuffer( void );
@@ -185,10 +183,10 @@ protected:
  */
 class FALAudioDevice : public FAudioDevice
 {
-public: 
+public:
 	FALAudioDevice();
 
-	virtual ~FALAudioDevice() {} 
+	virtual ~FALAudioDevice() {}
 
 	virtual FName GetRuntimeFormat(USoundWave* SoundWave) override
 	{
@@ -198,7 +196,7 @@ public:
 
 	/** Starts up any platform specific hardware/APIs */
 	virtual bool InitializeHardware() override;
- 
+
 	/**
 	 * Tears down audio device by stopping all sounds, removing all buffers,  destroying all sources, ... Called by both Destroy and ShutdownAfterError
 	 * to perform the actual tear down.
@@ -217,29 +215,23 @@ public:
 
 	virtual class ICompressedAudioInfo* CreateCompressedAudioInfo(USoundWave* SoundWave) override;
 
-	void FindProcs( bool AllowExt );
-
 	// Error checking.
 	bool alError( const TCHAR* Text, bool Log = true );
 
-	/** 
+	/**
 	 * Makes sure correct context is set.
-	 * 
+	 *
 	 * @param CallSiteIdentifier identifies the call site for debugging/error reporting purposes
 	 */
 	void MakeCurrent(const TCHAR * CallSiteIdentifier = nullptr);
 
 protected:
 
-	virtual FSoundSource* CreateSoundSource() override; 
+	virtual FSoundSource* CreateSoundSource() override;
 
 	/** Returns the enum for the internal format for playing a sound with this number of channels. */
 	ALuint GetInternalFormat( int NumChannels );
 
-	// Dynamic binding.
-	void FindProc( void*& ProcAddress, char* Name, char* SupportName, bool& Supports, bool AllowExt );
-	bool FindExt( const TCHAR* Name );
-
 
 	// Variables.
 
