PR #1206: Do not crash on check() if OpenAL initialization fails (UE-16460).

- Contributed by [REDACTED_PERSON].
- HardwareDevice can be null if the initialization (e.g. alcOpenDevice()) fails.

diff --git a/Engine/Source/Runtime/Linux/ALAudio/Private/ALAudioDevice.cpp b/Engine/Source/Runtime/Linux/ALAudio/Private/ALAudioDevice.cpp
--- a/Engine/Source/Runtime/Linux/ALAudio/Private/ALAudioDevice.cpp
+++ b/Engine/Source/Runtime/Linux/ALAudio/Private/ALAudioDevice.cpp
@@ -84,12 +84,14 @@ void FALAudioDevice::TeardownHardware( void )
 	// Close the hardware device
 	if( &alcCloseDevice )
 	{
-		checkf(HardwareDevice, TEXT("Tearing down invalid OpenAL device! (HardwareDevice should not be null)."));
-		const ALCchar* DeviceName = alcGetString(HardwareDevice, ALC_DEVICE_SPECIFIER);
-		UE_LOG(LogALAudio, Log, TEXT("Closing ALAudio device : %s"), StringCast<TCHAR>(static_cast<const ANSICHAR*>(DeviceName)).Get());
+		if (HardwareDevice)
+		{
+			const ALCchar* DeviceName = alcGetString(HardwareDevice, ALC_DEVICE_SPECIFIER);
+			UE_LOG(LogALAudio, Log, TEXT("Closing ALAudio device : %s"), StringCast<TCHAR>(static_cast<const ANSICHAR*>(DeviceName)).Get());
 
-		alcCloseDevice( HardwareDevice );
-		HardwareDevice = nullptr;
+			alcCloseDevice(HardwareDevice);
+			HardwareDevice = nullptr;
+		}
 	}
 
 }
