UE-9692, UE-9446 Fixing 3d audio spatialization for HTML5 and Linux.
- 3d audio was fundamnetally broken in HTML5 and Linux for a long time (probably since original port from UE3)
- Fix is to do 3d spatialization math the same way we do it on XAudio2 (i.e. transform to listener-space then do relative source position from there)
- Since HTML5 and Linux *both* use nearly identical OpenAL implementations (and Linux has been fixed by users for a few issues), I decided to nuke the HTML5 implementation and instead link to the Linux implementation. Not sure if we should move the linux OpenAL implementation to a different folder, but it turned out not to matter for packing and testing HTML5
- Also went ahead and added a -3dB headroom factor to reduce volume overall with OpenAL implementation since it was louder than the XAudio2 implementation. I didn't do any technical measurements to come up with this headroom factor, just ear-balled it. We can do something more precise if needed.
- Tested this fix using a simple test-map of a single looping sound source. It was obviously broken (i.e. not panning) and is now obviously fixed.

diff --git a/Engine/Config/HTML5/HTML5Engine.ini b/Engine/Config/HTML5/HTML5Engine.ini
--- a/Engine/Config/HTML5/HTML5Engine.ini
+++ b/Engine/Config/HTML5/HTML5Engine.ini
@@ -1,5 +1,5 @@
 [Audio]
-AudioDeviceModuleName=HTML5Audio
+AudioDeviceModuleName=ALAudio
 
 [OnlineSubsystemMcp]
 benabled=false
diff --git a/Engine/Source/Editor/UnrealEd/UnrealEd.Build.cs b/Engine/Source/Editor/UnrealEd/UnrealEd.Build.cs
--- a/Engine/Source/Editor/UnrealEd/UnrealEd.Build.cs
+++ b/Engine/Source/Editor/UnrealEd/UnrealEd.Build.cs
@@ -255,7 +255,7 @@ public class UnrealEd : ModuleRules
 
         if (Target.Platform == UnrealTargetPlatform.HTML5)
         {
-            PublicDependencyModuleNames.Add("HTML5Audio");
+			PublicDependencyModuleNames.Add("ALAudio");
         }
 
 		AddThirdPartyPrivateStaticDependencies(Target,
diff --git a/Engine/Source/Runtime/HTML5/HTML5Audio/HTML5Audio.Build.cs b/Engine/Source/Runtime/HTML5/HTML5Audio/HTML5Audio.Build.cs
deleted file mode 100644
diff --git a/Engine/Source/Runtime/HTML5/HTML5Audio/Private/HTML5AudioBuffer.cpp b/Engine/Source/Runtime/HTML5/HTML5Audio/Private/HTML5AudioBuffer.cpp
deleted file mode 100644
diff --git a/Engine/Source/Runtime/HTML5/HTML5Audio/Private/HTML5AudioDevice.cpp b/Engine/Source/Runtime/HTML5/HTML5Audio/Private/HTML5AudioDevice.cpp
deleted file mode 100644
diff --git a/Engine/Source/Runtime/HTML5/HTML5Audio/Private/HTML5AudioSource.cpp b/Engine/Source/Runtime/HTML5/HTML5Audio/Private/HTML5AudioSource.cpp
deleted file mode 100644
diff --git a/Engine/Source/Runtime/HTML5/HTML5Audio/Public/HTML5AudioDevice.h b/Engine/Source/Runtime/HTML5/HTML5Audio/Public/HTML5AudioDevice.h
deleted file mode 100644
diff --git a/Engine/Source/Runtime/Launch/Launch.Build.cs b/Engine/Source/Runtime/Launch/Launch.Build.cs
--- a/Engine/Source/Runtime/Launch/Launch.Build.cs
+++ b/Engine/Source/Runtime/Launch/Launch.Build.cs
@@ -181,7 +181,7 @@ public class Launch : ModuleRules
 
         if (Target.Platform == UnrealTargetPlatform.HTML5 )
         {
-            PrivateDependencyModuleNames.Add("HTML5Audio");
+			PrivateDependencyModuleNames.Add("ALAudio");
 			if (Target.Architecture == "-win32")
 			{
                 PrivateDependencyModuleNames.Add("HTML5Win32");
diff --git a/Engine/Source/Runtime/Linux/ALAudio/Private/ALAudioDevice.cpp b/Engine/Source/Runtime/Linux/ALAudio/Private/ALAudioDevice.cpp
--- a/Engine/Source/Runtime/Linux/ALAudio/Private/ALAudioDevice.cpp
+++ b/Engine/Source/Runtime/Linux/ALAudio/Private/ALAudioDevice.cpp
@@ -212,48 +212,17 @@ bool FALAudioDevice::InitializeHardware( void )
 /**
  * Update the audio device and calculates the cached inverse transform later
  * on used for spatialization.
- *
- * @param	Realtime	whether we are paused or not
  */
-void FALAudioDevice::Update( bool Realtime )
+void FALAudioDevice::UpdateHardware()
 {
-	FVector ListenerLocation	= Listeners[ 0 ].Transform.GetLocation();
-	FVector ListenerFront		= Listeners[ 0 ].GetFront();
-	FVector ListenerUp			= Listeners[ 0 ].GetUp();
-
-	// Set Player position
-	FVector Location;
-
-	// See file header for coordinate system explanation.
-	Location.X = ListenerLocation.X;
-	Location.Y = ListenerLocation.Z; // Z/Y swapped on purpose, see file header
-	Location.Z = ListenerLocation.Y; // Z/Y swapped on purpose, see file header
-	Location *= AUDIO_DISTANCE_FACTOR;
-
-	// Set Player orientation.
-	FVector Orientation[2];
-
-	// See file header for coordinate system explanation.
-	Orientation[0].X = ListenerFront.X;
-	Orientation[0].Y = ListenerFront.Z; // Z/Y swapped on purpose, see file header
-	Orientation[0].Z = ListenerFront.Y; // Z/Y swapped on purpose, see file header
-
-	// See file header for coordinate system explanation.
-	Orientation[1].X = ListenerUp.X;
-	Orientation[1].Y = ListenerUp.Z; // Z/Y swapped on purpose, see file header
-	Orientation[1].Z = ListenerUp.Y; // Z/Y swapped on purpose, see file header
-
-	// Make the listener still and the sounds move relatively -- this allows
-	// us to scale the doppler effect on a per-sound basis.
-	FVector Velocity = FVector( 0.0f, 0.0f, 0.0f );
-
-	MakeCurrent(TEXT("FALAudioDevice::Update()"));
-
-	alListenerfv( AL_POSITION, ( ALfloat* )&Location );
-	alListenerfv( AL_ORIENTATION, ( ALfloat* )&Orientation[0] );
-	alListenerfv( AL_VELOCITY, ( ALfloat* )&Velocity );
-
-	alError( TEXT( "UALAudioDevice::Update" ) );
+	if (Listeners.Num() > 0)
+	{
+		// Caches the matrix used to transform a sounds position into local space so we can just look
+		// at the Y component after normalization to determine spatialization.
+		const FVector Up = Listeners[0].GetUp();
+		const FVector Right = Listeners[0].GetFront();
+		InverseTransform = FMatrix(Up, Right, Up ^ Right, Listeners[0].Transform.GetTranslation()).InverseFast();
+	}
 }
 
 ALuint FALAudioDevice::GetInternalFormat( int NumChannels )
@@ -293,11 +262,13 @@ ALuint FALAudioDevice::GetInternalFormat( int NumChannels )
 
 void FALAudioDevice::MakeCurrent(const TCHAR * CallSiteIdentifier)
 {
+#if PLATFORM_LINUX
 	checkf(SoundContext, TEXT("Unitialiized sound context in FALAudioDevice::MakeCurrent()!"));
 	if (!alcMakeContextCurrent(SoundContext))
 	{
 		alError(CallSiteIdentifier ? CallSiteIdentifier : TEXT("FALAudioDevice::MakeCurrent()"));
 	}
+#endif
 }
 
 /*------------------------------------------------------------------------------------
diff --git a/Engine/Source/Runtime/Linux/ALAudio/Private/ALAudioSource.cpp b/Engine/Source/Runtime/Linux/ALAudio/Private/ALAudioSource.cpp
--- a/Engine/Source/Runtime/Linux/ALAudio/Private/ALAudioSource.cpp
+++ b/Engine/Source/Runtime/Linux/ALAudio/Private/ALAudioSource.cpp
@@ -100,6 +100,7 @@ void FALSoundSource::Update( void )
 			Volume *= 1.25f;
 		}
 		Volume *= FApp::GetVolumeMultiplier();
+		Volume *= AUDIO_OPEN_AL_VOLUME_SCALE;
 		Volume = FMath::Clamp(Volume, 0.0f, MAX_VOLUME);
 	}
 
@@ -128,20 +129,27 @@ void FALSoundSource::Update( void )
 	Velocity *= AUDIO_DISTANCE_FACTOR;
 
 	// We're using a relative coordinate system for un- spatialized sounds.
-	if( !WaveInstance->bUseSpatialization )
+	FVector RelativeDirection = FVector::ZeroVector;
+	if (WaveInstance->bUseSpatialization)
 	{
-		Location = FVector( 0.f, 0.f, 0.f );
+		FVector UnnormalizedDirection = ((FALAudioDevice*)AudioDevice)->InverseTransform.TransformPosition(WaveInstance->Location);
+		RelativeDirection = UnnormalizedDirection.GetSafeNormal();
+	}
+	else
+	{
+		Location = FVector(0.f, 0.f, 0.f);
 	}
 
-	alSourcef( SourceId, AL_GAIN, Volume );
-	alSourcef( SourceId, AL_PITCH, Pitch );
+	FVector EmitterPosition;
+	EmitterPosition.X = RelativeDirection.Y;
+	EmitterPosition.Y = RelativeDirection.X;
+	EmitterPosition.Z = -RelativeDirection.Z;
 
-	alSourcefv( SourceId, AL_POSITION, ( ALfloat* )&Location );
-	alSourcefv( SourceId, AL_VELOCITY, ( ALfloat* )&Velocity );
+	alSourcef(SourceId, AL_GAIN, Volume);
+	alSourcef(SourceId, AL_PITCH, Pitch);
 
-	// Platform dependent call to update the sound output with new parameters
-	// @todo openal: Is this no longer needed?
-	/// AudioDevice->UpdateEffect( this );
+	alSourcefv(SourceId, AL_POSITION, (ALfloat*)&EmitterPosition);
+	alSourcefv(SourceId, AL_VELOCITY, (ALfloat*)&Velocity);
 }
 
 /**
diff --git a/Engine/Source/Runtime/Linux/ALAudio/Public/ALAudioDevice.h b/Engine/Source/Runtime/Linux/ALAudio/Public/ALAudioDevice.h
--- a/Engine/Source/Runtime/Linux/ALAudio/Public/ALAudioDevice.h
+++ b/Engine/Source/Runtime/Linux/ALAudio/Public/ALAudioDevice.h
@@ -33,6 +33,10 @@ DECLARE_LOG_CATEGORY_EXTERN(LogALAudio, Log, All);
 // <=> 1 UU == 0.0127 m
 #define AUDIO_DISTANCE_FACTOR ( 0.0127f )
 
+// Audio volume is higher in OpenAL than XAudio2, we'll add a -3dB attenuation as headroom
+// to attempt to make it more on-par.
+#define AUDIO_OPEN_AL_VOLUME_SCALE (0.707f)
+
 /**
  * OpenAL implementation of FSoundBuffer, containing the wave data and format information.
  */
@@ -206,10 +210,8 @@ public:
 	/**
 	 * Update the audio device and calculates the cached inverse transform later
 	 * on used for spatialization.
-	 *
-	 * @param	Realtime	whether we are paused or not
 	 */
-	virtual void Update( bool bGameTicking );
+	virtual void UpdateHardware() override;
 
 	virtual bool HasCompressedAudioInfoClass(USoundWave* SoundWave) override;
 
@@ -248,7 +250,11 @@ protected:
 	ALenum										Surround61Format;
 	ALenum										Surround71Format;
 
+	/** Inverse listener transformation, used for spatialization */
+	FMatrix										InverseTransform;
+
 	friend class FALSoundBuffer;
+	friend class FALSoundSource;
 };
 
 FORCEINLINE FALAudioDevice * FALSoundSource::GetALDevice()
