Linux: moved common startup code to a separate module.

- The code that handles initial setup (command line processing, rlimits and debugger) is now shared between engine, UFE and SlateViewer.

diff --git a/Engine/Source/Programs/SlateViewer/Private/Linux/SlateViewerMainLinux.cpp b/Engine/Source/Programs/SlateViewer/Private/Linux/SlateViewerMainLinux.cpp
--- a/Engine/Source/Programs/SlateViewer/Private/Linux/SlateViewerMainLinux.cpp
+++ b/Engine/Source/Programs/SlateViewer/Private/Linux/SlateViewerMainLinux.cpp
@@ -1,25 +1,9 @@
 // Copyright 1998-2015 Epic Games, Inc. All Rights Reserved.
 
-
 #include "SlateViewerApp.h"
+#include "LinuxCommonStartup.h"
 
-static FString GSavedCommandLine;
-
-/**
- * WinMain, called when the application is started
- */
-int main( int argc, char** argv )
+int main(int argc, char *argv[])
 {
-	FPlatformMisc::SetGracefulTerminationHandler();
-
-	for (int32 Option = 1; Option < argc; Option++)
-	{
-		GSavedCommandLine += TEXT(" ");
-		GSavedCommandLine += UTF8_TO_TCHAR(argv[Option]);	// note: technically it depends on locale
-	}
-
-	// do the slate viewer thing
-	RunSlateViewer(*GSavedCommandLine);
-
-	return 0;
+	return CommonLinuxMain(argc, argv, &RunSlateViewer);
 }
diff --git a/Engine/Source/Programs/SlateViewer/Private/SlateViewerApp.cpp b/Engine/Source/Programs/SlateViewer/Private/SlateViewerApp.cpp
--- a/Engine/Source/Programs/SlateViewer/Private/SlateViewerApp.cpp
+++ b/Engine/Source/Programs/SlateViewer/Private/SlateViewerApp.cpp
@@ -20,7 +20,7 @@ namespace WorkspaceMenu
 }
 
 
-void RunSlateViewer( const TCHAR* CommandLine )
+int RunSlateViewer( const TCHAR* CommandLine )
 {
 	// start up the main loop
 	GEngineLoop.PreInit(CommandLine);
@@ -76,6 +76,8 @@ void RunSlateViewer( const TCHAR* CommandLine )
 	}
 	FModuleManager::Get().UnloadModulesAtShutdown();
 	FSlateApplication::Shutdown();
+
+	return 0;
 }
 
 namespace
diff --git a/Engine/Source/Programs/SlateViewer/Public/SlateViewerApp.h b/Engine/Source/Programs/SlateViewer/Public/SlateViewerApp.h
--- a/Engine/Source/Programs/SlateViewer/Public/SlateViewerApp.h
+++ b/Engine/Source/Programs/SlateViewer/Public/SlateViewerApp.h
@@ -9,7 +9,7 @@
 /**
  * Run the SlateViewer .
  */
-void RunSlateViewer(const TCHAR* Commandline);
+int RunSlateViewer(const TCHAR* Commandline);
 
 /**
  * Spawn the contents of the web browser tab
diff --git a/Engine/Source/Programs/SlateViewer/SlateViewer.Build.cs b/Engine/Source/Programs/SlateViewer/SlateViewer.Build.cs
--- a/Engine/Source/Programs/SlateViewer/SlateViewer.Build.cs
+++ b/Engine/Source/Programs/SlateViewer/SlateViewer.Build.cs
@@ -54,5 +54,14 @@ public class SlateViewer : ModuleRules
                 }
             );
 		}
+
+		if (Target.Platform == UnrealTargetPlatform.Linux)
+		{
+			PrivateDependencyModuleNames.AddRange(
+				new string[] {
+					"LinuxCommonStartup"
+				}
+			);
+		}
 	}
 }
diff --git a/Engine/Source/Programs/UnrealFrontend/Private/Linux/LinuxUnrealFrontendMain.cpp b/Engine/Source/Programs/UnrealFrontend/Private/Linux/LinuxUnrealFrontendMain.cpp
--- a/Engine/Source/Programs/UnrealFrontend/Private/Linux/LinuxUnrealFrontendMain.cpp
+++ b/Engine/Source/Programs/UnrealFrontend/Private/Linux/LinuxUnrealFrontendMain.cpp
@@ -2,65 +2,10 @@
 
 #include "UnrealFrontendPrivatePCH.h"
 #include "UnrealFrontendMain.h"
-#include "ExceptionHandling.h"
 #include "RequiredProgramMainCPPInclude.h"
+#include "LinuxCommonStartup.h"
 
-static FString GSavedCommandLine;
-
-
-int main( int argc, char *argv[] )
+int main(int argc, char *argv[])
 {
-	for (int32 Option = 1; Option < argc; ++Option)
-	{
-		GSavedCommandLine += TEXT(" ");
-		FString Argument(ANSI_TO_TCHAR(argv[Option]));
-
-		if (Argument.Contains(TEXT(" ")))
-		{
-			if (Argument.Contains(TEXT("=")))
-			{
-				FString ArgName;
-				FString ArgValue;
-
-				Argument.Split( TEXT("="), &ArgName, &ArgValue );
-				Argument = FString::Printf( TEXT("%s=\"%s\""), *ArgName, *ArgValue );
-			}
-			else
-			{
-				Argument = FString::Printf(TEXT("\"%s\""), *Argument);
-			}
-		}
-
-		GSavedCommandLine += Argument;
-	}
-	
-	FPlatformMisc::SetGracefulTerminationHandler();
-
-#if !UE_BUILD_SHIPPING
-	if (FParse::Param(*GSavedCommandLine,TEXT("crashreports")))
-	{
-		GAlwaysReportCrash = true;
-	}
-#endif
-
-	GUseCrashReportClient = true;
-	
-#if UE_BUILD_DEBUG
-	if (!GAlwaysReportCrash)
-#else
-	if (FPlatformMisc::IsDebuggerPresent() && !GAlwaysReportCrash)
-#endif
-	{
-		UnrealFrontendMain(*GSavedCommandLine);
-	}
-	else
-	{
-		GIsGuarded = 1;
-		UnrealFrontendMain(*GSavedCommandLine);
-		GIsGuarded = 0;
-	}
-
-	FEngineLoop::AppExit();
-	
-	return 0;
+	return CommonLinuxMain(argc, argv, &UnrealFrontendMain);
 }
diff --git a/Engine/Source/Programs/UnrealFrontend/UnrealFrontend.Build.cs b/Engine/Source/Programs/UnrealFrontend/UnrealFrontend.Build.cs
--- a/Engine/Source/Programs/UnrealFrontend/UnrealFrontend.Build.cs
+++ b/Engine/Source/Programs/UnrealFrontend/UnrealFrontend.Build.cs
@@ -59,6 +59,14 @@ public class UnrealFrontend : ModuleRules
 		{
 			PrivateDependencyModuleNames.Add("VisualStudioSourceCodeAccess");
 		}
+		else if (Target.Platform == UnrealTargetPlatform.Linux)
+		{
+			PrivateDependencyModuleNames.AddRange(
+				new string[] {
+					"LinuxCommonStartup"
+				}
+			);
+		}
 
 		// @todo: allow for better plug-in support in standalone Slate apps
 		PrivateDependencyModuleNames.AddRange(
diff --git a/Engine/Source/Runtime/Launch/Launch.Build.cs b/Engine/Source/Runtime/Launch/Launch.Build.cs
--- a/Engine/Source/Runtime/Launch/Launch.Build.cs
+++ b/Engine/Source/Runtime/Launch/Launch.Build.cs
@@ -205,5 +205,14 @@ public class Launch : ModuleRules
 		{
 			bFasterWithoutUnity = true;
 		}
+
+		if (Target.Platform == UnrealTargetPlatform.Linux)
+		{
+			PrivateDependencyModuleNames.AddRange(
+				new string[] {
+					"LinuxCommonStartup"
+				}
+			);
+		}
 	}
 }
diff --git a/Engine/Source/Runtime/Launch/Private/Linux/LaunchLinux.cpp b/Engine/Source/Runtime/Launch/Private/Linux/LaunchLinux.cpp
--- a/Engine/Source/Runtime/Launch/Private/Linux/LaunchLinux.cpp
+++ b/Engine/Source/Runtime/Launch/Private/Linux/LaunchLinux.cpp
@@ -1,203 +1,9 @@
 // Copyright 1998-2015 Epic Games, Inc. All Rights Reserved.
 
 #include "LaunchPrivatePCH.h"
-#include "ExceptionHandling.h"
-#include "LinuxPlatformCrashContext.h"
-
-#if WITH_EDITOR
-#include "Editor/MainFrame/Private/Frame/MainFrameActions.h"
-#endif
-
-#include <locale.h>
-#include <sys/resource.h>
-
-static FString GSavedCommandLine;
-extern int32 GuardedMain( const TCHAR* CmdLine );
-extern void LaunchStaticShutdownAfterError();
-
-// FIXME: handle expose it someplace else?
-extern int32 DLLIMPORT ReportCrash(const FLinuxCrashContext& Context);
-extern void DLLIMPORT GenerateCrashInfoAndLaunchReporter(const FLinuxCrashContext& Context);
-
-/**
- * Game-specific crash reporter
- */
-void EngineCrashHandler(const FGenericCrashContext& GenericContext)
-{
-	const FLinuxCrashContext& Context = static_cast< const FLinuxCrashContext& >( GenericContext );
-
-	printf("EngineCrashHandler: Signal=%d\n", Context.Signal);
-	ReportCrash(Context);
-	if (GLog)
-	{
-		GLog->Flush();
-	}
-	if (GWarn)
-	{
-		GWarn->Flush();
-	}
-	if (GError)
-	{
-		GError->Flush();
-		GError->HandleError();
-	}
-	//LaunchStaticShutdownAfterError();	// this tends to crash itself, so while it should be here, it needs to be commented
-	return GenerateCrashInfoAndLaunchReporter(Context);
-}
-
-
-/**
- * Increases (soft) limit on a specific resource
- *
- * @param DesiredLimit - max number of open files desired.
- */
-static bool IncreaseLimit(int Resource, rlim_t DesiredLimit)
-{
-	rlimit Limit;
-	if (getrlimit(Resource, &Limit) != 0)
-	{
-		fprintf(stderr, "getrlimit() failed with error %d (%s)\n", errno, strerror(errno));
-		return false;
-	}
-
-	if (Limit.rlim_cur == RLIM_INFINITY || Limit.rlim_cur >= DesiredLimit)
-	{
-#if !UE_BUILD_SHIPPING
-		printf("- Existing per-process limit (soft=%lu, hard=%lu) is enough for us (need only %lu)\n", Limit.rlim_cur, Limit.rlim_max, DesiredLimit);
-#endif // !UE_BUILD_SHIPPING
-		return true;
-	}
-
-	Limit.rlim_cur = DesiredLimit;
-	if (setrlimit(Resource, &Limit) != 0)
-	{
-		fprintf(stderr, "setrlimit() failed with error %d (%s)\n", errno, strerror(errno));
-
-		if (errno == EINVAL)
-		{
-			if (DesiredLimit == RLIM_INFINITY)
-			{
-				fprintf(stderr, "- Max per-process value allowed is %lu (we wanted infinity).\n", Limit.rlim_max);
-			}
-			else
-			{
-				fprintf(stderr, "- Max per-process value allowed is %lu (we wanted %lu).\n", Limit.rlim_max, DesiredLimit);
-			}
-		}
-		return false;
-	}
-
-	return true;
-}
-
-/**
- * Expects GSavedCommandLine to be set up. Increases limit on 
- *  - number of open files to be no less than desired (if specified on command line, otherwise left alone)
- *  - size of core file, so core gets dumped and we can debug crashed builds (unless overridden with -nocore)
- *
- */
-static bool IncreasePerProcessLimits()
-{
-	// honor the parameter if given, but don't change limits if not
-	int32 FileHandlesToReserve = -1;
-	if (FParse::Value(*GSavedCommandLine, TEXT("numopenfiles="), FileHandlesToReserve) && FileHandlesToReserve > 0)
-	{
-#if !UE_BUILD_SHIPPING
-		printf("Increasing per-process limit of open file handles to %d\n", FileHandlesToReserve);
-#endif // !UE_BUILD_SHIPPING
-		if (!IncreaseLimit(RLIMIT_NOFILE, FileHandlesToReserve))
-		{
-			fprintf(stderr, "Could not adjust number of file handles, consider changing \"nofile\" in /etc/security/limits.conf and relogin.\nerror(%d): %s\n", errno, strerror(errno));
-			return false;
-		}
-	}
-
-#if !UE_BUILD_SHIPPING
-	if (!FParse::Param(*GSavedCommandLine, TEXT("nocore")))
-	{
-		printf("Increasing per-process limit of core file size to infinity.\n");
-		if (!IncreaseLimit(RLIMIT_CORE, RLIM_INFINITY))
-		{
-			fprintf(stderr, "Could not adjust core file size, consider changing \"core\" in /etc/security/limits.conf and relogin.\nerror(%d): %s\n", errno, strerror(errno));
-			fprintf(stderr, "Alternatively, pass -nocore if you are unable or unwilling to do that.\n");
-			return false;
-		}
-	}
-#endif // !UE_BUILD_SHIPPING
-
-	return true;
-}
-
+#include "LinuxCommonStartup.h"
 
 int main(int argc, char *argv[])
 {
-	FPlatformMisc::SetGracefulTerminationHandler();
-
-#if UE_BUILD_SHIPPING
-	// only printed in shipping
-	printf("%s %d %d %d %d\n", StringCast<ANSICHAR>(*FEngineVersion::Current().ToString()).Get(), GEngineMinNetVersion, GEngineNegotiationVersion, GPackageFileUE4Version, GPackageFileLicenseeUE4Version);
-#endif // UE_BUILD_SHIPPING
-	int ErrorLevel = 0;
-
-	setlocale(LC_CTYPE, "");
-	for (int32 Option = 1; Option < argc; Option++)
-	{
-		GSavedCommandLine += TEXT(" ");
-		// we need to quote stuff that has spaces in it because something somewhere is removing quotation marks before they arrive here
-		FString Temp = UTF8_TO_TCHAR(argv[Option]);
-		if (Temp.Contains(TEXT(" ")))
-		{
-			if(Temp.StartsWith(TEXT("-")))
-			{
-				Temp = Temp.Replace(TEXT("="), TEXT("=\""));
-			}
-			else
-			{
-				Temp = TEXT("\"") + Temp;
-			}
-			Temp += TEXT("\"");
-		}
-		GSavedCommandLine += Temp; 	// note: technically it depends on locale
-	}
-
-#if !UE_BUILD_SHIPPING
-	GAlwaysReportCrash = true;	// set by default and reverse the behavior
-	if ( FParse::Param( *GSavedCommandLine,TEXT("nocrashreports") ) || FParse::Param( *GSavedCommandLine,TEXT("no-crashreports") ) )
-	{
-		GAlwaysReportCrash = false;
-	}
-#endif
-
-	if (!IncreasePerProcessLimits())
-	{
-		fprintf(stderr, "Could not set desired per-process limits, consider changing system limits.\n");
-		ErrorLevel = 1;
-	}
-	else
-	{
-#if UE_BUILD_DEBUG
-		if( true && !GAlwaysReportCrash )
-#else
-		if( FPlatformMisc::IsDebuggerPresent() && !GAlwaysReportCrash )
-#endif
-		{
-			// Don't use exception handling when a debugger is attached to exactly trap the crash. This does NOT check
-			// whether we are the first instance or not!
-			ErrorLevel = GuardedMain( *GSavedCommandLine );
-		}
-		else
-		{
-			FPlatformMisc::SetCrashHandler(EngineCrashHandler);
-			GIsGuarded = 1;
-			// Run the guarded code.
-			ErrorLevel = GuardedMain( *GSavedCommandLine );
-			GIsGuarded = 0;
-		}
-	}
-
-	if (ErrorLevel)
-	{
-		printf("Exiting abnormally (error code: %d)\n", ErrorLevel);
-	}
-	return ErrorLevel;
+	return CommonLinuxMain(argc, argv, &GuardedMain);
 }
diff --git a/Engine/Source/Runtime/Linux/LinuxCommonStartup/LinuxCommonStartup.Build.cs b/Engine/Source/Runtime/Linux/LinuxCommonStartup/LinuxCommonStartup.Build.cs
new file mode 100644
--- /dev/null
+++ b/Engine/Source/Runtime/Linux/LinuxCommonStartup/LinuxCommonStartup.Build.cs
@@ -0,0 +1,22 @@
+// Copyright 1998-2015 Epic Games, Inc. All Rights Reserved.
+
+using UnrealBuildTool;
+using System.IO;
+
+public class LinuxCommonStartup : ModuleRules
+{
+	public LinuxCommonStartup(TargetInfo Target)
+	{
+		PrivateDependencyModuleNames.Add("Core");
+		PrivateDependencyModuleNames.Add("Projects");
+		
+		if (Target.Type == TargetRules.TargetType.Editor)
+		{
+			PrivateDependencyModuleNames.AddRange(
+				new string[] {
+					"SourceControl",
+				}
+			);
+		}
+	}
+}
diff --git a/Engine/Source/Runtime/Linux/LinuxCommonStartup/Private/LinuxCommonStartup.cpp b/Engine/Source/Runtime/Linux/LinuxCommonStartup/Private/LinuxCommonStartup.cpp
new file mode 100644
--- /dev/null
+++ b/Engine/Source/Runtime/Linux/LinuxCommonStartup/Private/LinuxCommonStartup.cpp
@@ -0,0 +1,210 @@
+// Copyright 1998-2015 Epic Games, Inc. All Rights Reserved.
+
+#include "LinuxCommonStartup.h"
+#include "ExceptionHandling.h"
+#include "LinuxPlatformCrashContext.h"
+#include "ModuleManager.h"
+
+#include <locale.h>
+#include <sys/resource.h>
+
+static FString GSavedCommandLine;
+extern int32 GuardedMain( const TCHAR* CmdLine );
+extern void LaunchStaticShutdownAfterError();
+
+// FIXME: handle expose it someplace else?
+extern int32 DLLIMPORT ReportCrash(const FLinuxCrashContext& Context);
+extern void DLLIMPORT GenerateCrashInfoAndLaunchReporter(const FLinuxCrashContext& Context);
+
+/**
+ * Game-specific crash reporter
+ */
+void CommonLinuxCrashHandler(const FGenericCrashContext& GenericContext)
+{
+	const FLinuxCrashContext& Context = static_cast< const FLinuxCrashContext& >( GenericContext );
+
+	printf("EngineCrashHandler: Signal=%d\n", Context.Signal);
+	ReportCrash(Context);
+	if (GLog)
+	{
+		GLog->Flush();
+	}
+	if (GWarn)
+	{
+		GWarn->Flush();
+	}
+	if (GError)
+	{
+		GError->Flush();
+		GError->HandleError();
+	}
+
+	return GenerateCrashInfoAndLaunchReporter(Context);
+}
+
+
+/**
+ * Increases (soft) limit on a specific resource
+ *
+ * @param DesiredLimit - max number of open files desired.
+ */
+static bool IncreaseLimit(int Resource, rlim_t DesiredLimit)
+{
+	rlimit Limit;
+	if (getrlimit(Resource, &Limit) != 0)
+	{
+		fprintf(stderr, "getrlimit() failed with error %d (%s)\n", errno, strerror(errno));
+		return false;
+	}
+
+	if (Limit.rlim_cur == RLIM_INFINITY || Limit.rlim_cur >= DesiredLimit)
+	{
+#if !UE_BUILD_SHIPPING
+		printf("- Existing per-process limit (soft=%lu, hard=%lu) is enough for us (need only %lu)\n", Limit.rlim_cur, Limit.rlim_max, DesiredLimit);
+#endif // !UE_BUILD_SHIPPING
+		return true;
+	}
+
+	Limit.rlim_cur = DesiredLimit;
+	if (setrlimit(Resource, &Limit) != 0)
+	{
+		fprintf(stderr, "setrlimit() failed with error %d (%s)\n", errno, strerror(errno));
+
+		if (errno == EINVAL)
+		{
+			if (DesiredLimit == RLIM_INFINITY)
+			{
+				fprintf(stderr, "- Max per-process value allowed is %lu (we wanted infinity).\n", Limit.rlim_max);
+			}
+			else
+			{
+				fprintf(stderr, "- Max per-process value allowed is %lu (we wanted %lu).\n", Limit.rlim_max, DesiredLimit);
+			}
+		}
+		return false;
+	}
+
+	return true;
+}
+
+/**
+ * Expects GSavedCommandLine to be set up. Increases limit on
+ *  - number of open files to be no less than desired (if specified on command line, otherwise left alone)
+ *  - size of core file, so core gets dumped and we can debug crashed builds (unless overridden with -nocore)
+ *
+ */
+static bool IncreasePerProcessLimits()
+{
+	// honor the parameter if given, but don't change limits if not
+	int32 FileHandlesToReserve = -1;
+	if (FParse::Value(*GSavedCommandLine, TEXT("numopenfiles="), FileHandlesToReserve) && FileHandlesToReserve > 0)
+	{
+#if !UE_BUILD_SHIPPING
+		printf("Increasing per-process limit of open file handles to %d\n", FileHandlesToReserve);
+#endif // !UE_BUILD_SHIPPING
+		if (!IncreaseLimit(RLIMIT_NOFILE, FileHandlesToReserve))
+		{
+			fprintf(stderr, "Could not adjust number of file handles, consider changing \"nofile\" in /etc/security/limits.conf and relogin.\nerror(%d): %s\n", errno, strerror(errno));
+			return false;
+		}
+	}
+
+#if !UE_BUILD_SHIPPING
+	if (!FParse::Param(*GSavedCommandLine, TEXT("nocore")))
+	{
+		printf("Increasing per-process limit of core file size to infinity.\n");
+		if (!IncreaseLimit(RLIMIT_CORE, RLIM_INFINITY))
+		{
+			fprintf(stderr, "Could not adjust core file size, consider changing \"core\" in /etc/security/limits.conf and relogin.\nerror(%d): %s\n", errno, strerror(errno));
+			fprintf(stderr, "Alternatively, pass -nocore if you are unable or unwilling to do that.\n");
+			return false;
+		}
+	}
+#endif // !UE_BUILD_SHIPPING
+
+	return true;
+}
+
+
+int CommonLinuxMain(int argc, char *argv[], int (*RealMain)(const TCHAR * CommandLine))
+{
+	FPlatformMisc::SetGracefulTerminationHandler();
+
+#if UE_BUILD_SHIPPING
+	// only printed in shipping
+	printf("%s %d %d %d %d\n", StringCast<ANSICHAR>(*FEngineVersion::Current().ToString()).Get(), GEngineMinNetVersion, GEngineNegotiationVersion, GPackageFileUE4Version, GPackageFileLicenseeUE4Version);
+#endif // UE_BUILD_SHIPPING
+	int ErrorLevel = 0;
+
+	setlocale(LC_CTYPE, "");
+	for (int32 Option = 1; Option < argc; Option++)
+	{
+		GSavedCommandLine += TEXT(" ");
+		// we need to quote stuff that has spaces in it because something somewhere is removing quotation marks before they arrive here
+		FString Temp = UTF8_TO_TCHAR(argv[Option]);
+		if (Temp.Contains(TEXT(" ")))
+		{
+			if(Temp.StartsWith(TEXT("-")))
+			{
+				Temp = Temp.Replace(TEXT("="), TEXT("=\""));
+			}
+			else
+			{
+				Temp = TEXT("\"") + Temp;
+			}
+			Temp += TEXT("\"");
+		}
+		GSavedCommandLine += Temp; 	// note: technically it depends on locale
+	}
+
+#if !UE_BUILD_SHIPPING
+	GAlwaysReportCrash = true;	// set by default and reverse the behavior
+	if ( FParse::Param( *GSavedCommandLine,TEXT("nocrashreports") ) || FParse::Param( *GSavedCommandLine,TEXT("no-crashreports") ) )
+	{
+		GAlwaysReportCrash = false;
+	}
+#endif
+
+	if (!IncreasePerProcessLimits())
+	{
+		fprintf(stderr, "Could not set desired per-process limits, consider changing system limits.\n");
+		ErrorLevel = 1;
+	}
+	else
+	{
+#if UE_BUILD_DEBUG
+		if( true && !GAlwaysReportCrash )
+#else
+		if( FPlatformMisc::IsDebuggerPresent() && !GAlwaysReportCrash )
+#endif
+		{
+			// Don't use exception handling when a debugger is attached to exactly trap the crash. This does NOT check
+			// whether we are the first instance or not!
+			ErrorLevel = RealMain( *GSavedCommandLine );
+		}
+		else
+		{
+			FPlatformMisc::SetCrashHandler(CommonLinuxCrashHandler);
+			GIsGuarded = 1;
+			// Run the guarded code.
+			ErrorLevel = RealMain( *GSavedCommandLine );
+			GIsGuarded = 0;
+		}
+	}
+
+	if (ErrorLevel)
+	{
+		printf("Exiting abnormally (error code: %d)\n", ErrorLevel);
+	}
+	return ErrorLevel;
+}
+
+
+class FLinuxCommonStartupModule : public IModuleInterface
+{
+	/** IModuleInterface implementation */
+	virtual void StartupModule() override {};
+	virtual void ShutdownModule() override {};
+};
+
+IMPLEMENT_MODULE(FLinuxCommonStartupModule, LinuxCommonStartup);
diff --git a/Engine/Source/Runtime/Linux/LinuxCommonStartup/Public/LinuxCommonStartup.h b/Engine/Source/Runtime/Linux/LinuxCommonStartup/Public/LinuxCommonStartup.h
new file mode 100644
--- /dev/null
+++ b/Engine/Source/Runtime/Linux/LinuxCommonStartup/Public/LinuxCommonStartup.h
@@ -0,0 +1,14 @@
+// Copyright 1998-2015 Epic Games, Inc. All Rights Reserved.
+
+#pragma once
+
+#include "Core.h"
+
+/**
+ * @brief CommonLinuxMain - executes common startup code for Linux programs/engine
+ * @param argc - number of arguments in argv[]
+ * @param argv - array of arguments
+ * @param RealMain - the next main routine to call in chain
+ * @return error code to return to the OS
+ */
+int CommonLinuxMain(int argc, char *argv[], int (*RealMain)(const TCHAR * CommandLine));
