Linux: fix CIS (engine version).

diff --git a/Engine/Source/Runtime/Linux/LinuxCommonStartup/Private/LinuxCommonStartup.cpp b/Engine/Source/Runtime/Linux/LinuxCommonStartup/Private/LinuxCommonStartup.cpp
--- a/Engine/Source/Runtime/Linux/LinuxCommonStartup/Private/LinuxCommonStartup.cpp
+++ b/Engine/Source/Runtime/Linux/LinuxCommonStartup/Private/LinuxCommonStartup.cpp
@@ -4,6 +4,9 @@
 #include "ExceptionHandling.h"
 #include "LinuxPlatformCrashContext.h"
 #include "ModuleManager.h"
+#if UE_BUILD_SHIPPING
+	#include "EngineVersion.h"
+#endif // UE_BUILD_SHIPPING
 
 #include <locale.h>
 #include <sys/resource.h>
