Linux: disable core dumps in Test and Shipping by default.

- Also got rid of most #ifdef's so breakage of Shipping code path is more obvious.

diff --git a/Engine/Source/Runtime/Linux/LinuxCommonStartup/Private/LinuxCommonStartup.cpp b/Engine/Source/Runtime/Linux/LinuxCommonStartup/Private/LinuxCommonStartup.cpp
--- a/Engine/Source/Runtime/Linux/LinuxCommonStartup/Private/LinuxCommonStartup.cpp
+++ b/Engine/Source/Runtime/Linux/LinuxCommonStartup/Private/LinuxCommonStartup.cpp
@@ -4,9 +4,7 @@
 #include "ExceptionHandling.h"
 #include "LinuxPlatformCrashContext.h"
 #include "ModuleManager.h"
-#if UE_BUILD_SHIPPING
-	#include "EngineVersion.h"
-#endif // UE_BUILD_SHIPPING
+#include "EngineVersion.h"
 
 #include <locale.h>
 #include <sys/resource.h>
@@ -47,11 +45,13 @@ void CommonLinuxCrashHandler(const FGenericCrashContext& GenericContext)
 
 
 /**
- * Increases (soft) limit on a specific resource
+ * Sets (soft) limit on a specific resource
  *
- * @param DesiredLimit - max number of open files desired.
+ * @param Resource - one of RLIMIT_* values
+ * @param DesiredLimit - desired value
+ * @param bIncreaseOnly - avoid changing the limit if current value is sufficient
  */
-static bool IncreaseLimit(int Resource, rlim_t DesiredLimit)
+bool SetResourceLimit(int Resource, rlim_t DesiredLimit, bool bIncreaseOnly)
 {
 	rlimit Limit;
 	if (getrlimit(Resource, &Limit) != 0)
@@ -60,11 +60,12 @@ static bool IncreaseLimit(int Resource, rlim_t DesiredLimit)
 		return false;
 	}
 
-	if (Limit.rlim_cur == RLIM_INFINITY || Limit.rlim_cur >= DesiredLimit)
+	if (bIncreaseOnly && (Limit.rlim_cur == RLIM_INFINITY || Limit.rlim_cur >= DesiredLimit))
 	{
-#if !UE_BUILD_SHIPPING
-		printf("- Existing per-process limit (soft=%lu, hard=%lu) is enough for us (need only %lu)\n", Limit.rlim_cur, Limit.rlim_max, DesiredLimit);
-#endif // !UE_BUILD_SHIPPING
+		if (!UE_BUILD_SHIPPING)
+		{
+			printf("- Existing per-process limit (soft=%lu, hard=%lu) is enough for us (need only %lu)\n", Limit.rlim_cur, Limit.rlim_max, DesiredLimit);
+		}
 		return true;
 	}
 
@@ -102,28 +103,51 @@ static bool IncreasePerProcessLimits()
 	int32 FileHandlesToReserve = -1;
 	if (FParse::Value(*GSavedCommandLine, TEXT("numopenfiles="), FileHandlesToReserve) && FileHandlesToReserve > 0)
 	{
-#if !UE_BUILD_SHIPPING
-		printf("Increasing per-process limit of open file handles to %d\n", FileHandlesToReserve);
-#endif // !UE_BUILD_SHIPPING
-		if (!IncreaseLimit(RLIMIT_NOFILE, FileHandlesToReserve))
+		if (!UE_BUILD_SHIPPING)
+		{
+			printf("Increasing per-process limit of open file handles to %d\n", FileHandlesToReserve);
+		}
+
+		if (!SetResourceLimit(RLIMIT_NOFILE, FileHandlesToReserve, true))
 		{
 			fprintf(stderr, "Could not adjust number of file handles, consider changing \"nofile\" in /etc/security/limits.conf and relogin.\nerror(%d): %s\n", errno, strerror(errno));
 			return false;
 		}
 	}
 
-#if !UE_BUILD_SHIPPING
-	if (!FParse::Param(*GSavedCommandLine, TEXT("nocore")))
+	// core dump policy:
+	// - Shipping and Test disable by default (unless -core is passed)
+	// - The rest set it to infinity unless -nocore is passed
+	// (in all scenarios user wish as expressed with -core or -nocore takes priority)
+	bool bDisableCore = (UE_BUILD_SHIPPING || UE_BUILD_TEST);
+	if (FParse::Param(*GSavedCommandLine, TEXT("nocore")))
+	{
+		bDisableCore = true;
+	}
+	if (FParse::Param(*GSavedCommandLine, TEXT("core")))
+	{
+		bDisableCore = false;
+	}
+
+	if (bDisableCore)
+	{
+		printf("Disabling core dumps.\n");
+		if (!SetResourceLimit(RLIMIT_CORE, 0, false))
+		{
+			fprintf(stderr, "Could not set core file size to 0, error(%d): %s\n", errno, strerror(errno));
+			return false;
+		}
+	}
+	else
 	{
 		printf("Increasing per-process limit of core file size to infinity.\n");
-		if (!IncreaseLimit(RLIMIT_CORE, RLIM_INFINITY))
+		if (!SetResourceLimit(RLIMIT_CORE, RLIM_INFINITY, true))
 		{
 			fprintf(stderr, "Could not adjust core file size, consider changing \"core\" in /etc/security/limits.conf and relogin.\nerror(%d): %s\n", errno, strerror(errno));
 			fprintf(stderr, "Alternatively, pass -nocore if you are unable or unwilling to do that.\n");
 			return false;
 		}
 	}
-#endif // !UE_BUILD_SHIPPING
 
 	return true;
 }
@@ -133,10 +157,12 @@ int CommonLinuxMain(int argc, char *argv[], int (*RealMain)(const TCHAR * Comman
 {
 	FPlatformMisc::SetGracefulTerminationHandler();
 
-#if UE_BUILD_SHIPPING
-	// only printed in shipping
-	printf("%s %d %d %d %d\n", StringCast<ANSICHAR>(*FEngineVersion::Current().ToString()).Get(), GEngineMinNetVersion, GEngineNegotiationVersion, GPackageFileUE4Version, GPackageFileLicenseeUE4Version);
-#endif // UE_BUILD_SHIPPING
+	if (UE_BUILD_SHIPPING)
+	{
+		// only printed in shipping
+		printf("%s %d %d %d %d\n", StringCast<ANSICHAR>(*FEngineVersion::Current().ToString()).Get(), GEngineMinNetVersion, GEngineNegotiationVersion, GPackageFileUE4Version, GPackageFileLicenseeUE4Version);
+	}
+
 	int ErrorLevel = 0;
 
 	setlocale(LC_CTYPE, "");
@@ -160,13 +186,14 @@ int CommonLinuxMain(int argc, char *argv[], int (*RealMain)(const TCHAR * Comman
 		GSavedCommandLine += Temp; 	// note: technically it depends on locale
 	}
 
-#if !UE_BUILD_SHIPPING
-	GAlwaysReportCrash = true;	// set by default and reverse the behavior
-	if ( FParse::Param( *GSavedCommandLine,TEXT("nocrashreports") ) || FParse::Param( *GSavedCommandLine,TEXT("no-crashreports") ) )
+	if (!UE_BUILD_SHIPPING)
 	{
-		GAlwaysReportCrash = false;
+		GAlwaysReportCrash = true;	// set by default and reverse the behavior
+		if ( FParse::Param( *GSavedCommandLine,TEXT("nocrashreports") ) || FParse::Param( *GSavedCommandLine,TEXT("no-crashreports") ) )
+		{
+			GAlwaysReportCrash = false;
+		}
 	}
-#endif
 
 	if (!IncreasePerProcessLimits())
 	{
