Fixed missed platform group defines

diff --git a/Engine/Source/Runtime/Linux/AudioMixerSDL/Private/AudioMixerPlatformSDL.cpp b/Engine/Source/Runtime/Linux/AudioMixerSDL/Private/AudioMixerPlatformSDL.cpp
--- a/Engine/Source/Runtime/Linux/AudioMixerSDL/Private/AudioMixerPlatformSDL.cpp
+++ b/Engine/Source/Runtime/Linux/AudioMixerSDL/Private/AudioMixerPlatformSDL.cpp
@@ -103,7 +103,7 @@ namespace Audio
 		SDL_AudioSpec DesiredSpec;
 		DesiredSpec.freq = PlatformSettings.SampleRate;
 
-#if PLATFORM_LINUX
+#if PLATFORM_UNIX
 		DesiredSpec.format = AUDIO_F32;
 		DesiredSpec.channels = 6;
 #else
@@ -145,7 +145,7 @@ namespace Audio
 		OutInfo.Name = OutInfo.DeviceId;
 		OutInfo.SampleRate = ActualSpec.freq;
 		
-#if PLATFORM_LINUX
+#if PLATFORM_UNIX
 		OutInfo.Format = EAudioMixerStreamDataFormat::Float;
 #else
 		// HTML5 supports s16 format only
@@ -192,7 +192,7 @@ namespace Audio
 			return false;
 		}
 
-#if PLATFORM_LINUX
+#if PLATFORM_UNIX
 		AudioSpecPrefered.format = AUDIO_F32;
 #else
 		// HTML5 supports s16 format only
@@ -224,7 +224,7 @@ namespace Audio
 		check(AudioSpecReceived.samples == OpenStreamParams.NumFrames);
 
 		// Compute the expected output byte length
-#if PLATFORM_LINUX
+#if PLATFORM_UNIX
 		OutputBufferByteLength = OpenStreamParams.NumFrames * AudioStreamInfo.DeviceInfo.NumChannels * sizeof(float);
 #else
 		// HTML5 supports s16 format only
