Allow for overriding the base AudioMixerSDL platform

diff --git a/Engine/Source/Runtime/Linux/AudioMixerSDL/AudioMixerSDL.Build.cs b/Engine/Source/Runtime/Linux/AudioMixerSDL/AudioMixerSDL.Build.cs
--- a/Engine/Source/Runtime/Linux/AudioMixerSDL/AudioMixerSDL.Build.cs
+++ b/Engine/Source/Runtime/Linux/AudioMixerSDL/AudioMixerSDL.Build.cs
@@ -12,6 +12,8 @@ public class AudioMixerSDL : ModuleRules
 		PublicIncludePaths.Add("Runtime/AudioMixer/Public");
 		PrivateIncludePaths.Add("Runtime/AudioMixer/Private");
 
+		PrivateIncludePaths.Add("Runtime/Linux/AudioMixerSDL/Private/" + Target.Platform);
+
 		PrivateDependencyModuleNames.AddRange(
 			new string[] {
 				"Core",
diff --git a/Engine/Source/Runtime/Linux/AudioMixerSDL/Private/AudioMixerPlatformSDL.cpp b/Engine/Source/Runtime/Linux/AudioMixerSDL/Private/AudioMixerPlatformSDL.cpp
--- a/Engine/Source/Runtime/Linux/AudioMixerSDL/Private/AudioMixerPlatformSDL.cpp
+++ b/Engine/Source/Runtime/Linux/AudioMixerSDL/Private/AudioMixerPlatformSDL.cpp
@@ -114,7 +114,7 @@ namespace Audio
 		DesiredSpec.format = AUDIO_S16;
 		DesiredSpec.channels = 2;
 #endif
-		
+
 		DesiredSpec.samples = PlatformSettings.CallbackBufferFrameSize;
 		DesiredSpec.callback = OnBufferEnd;
 		DesiredSpec.userdata = (void*)this;
@@ -213,7 +213,15 @@ namespace Audio
 			DeviceName = SDL_GetAudioDeviceName(OpenStreamParams.OutputDeviceIndex, 0);
 		}
 
-		AudioDeviceID = SDL_OpenAudioDevice(DeviceName, 0, &AudioSpecPrefered, &AudioSpecReceived, 0);
+		FString CurrentDeviceName = GetCurrentDeviceName();
+		if (CurrentDeviceName.Len() <= 0)
+		{
+			AudioDeviceID = SDL_OpenAudioDevice(DeviceName, 0, &AudioSpecPrefered, &AudioSpecReceived, 0);
+		}
+		else
+		{
+			AudioDeviceID = SDL_OpenAudioDevice(TCHAR_TO_ANSI(*CurrentDeviceName), 0, &AudioSpecPrefered, &AudioSpecReceived, 0);
+		}
 
 		if (!AudioDeviceID)
 		{
@@ -240,6 +248,11 @@ namespace Audio
 		return true;
 	}
 
+	FString FMixerPlatformSDL::GetCurrentDeviceName() const
+	{
+		return {};
+	}
+
 	bool FMixerPlatformSDL::CloseAudioStream()
 	{
 		if (AudioStreamInfo.StreamState == EAudioOutputStreamState::Closed)
diff --git a/Engine/Source/Runtime/Linux/AudioMixerSDL/Private/AudioMixerPlatformSDL.h b/Engine/Source/Runtime/Linux/AudioMixerSDL/Private/AudioMixerPlatformSDL.h
--- a/Engine/Source/Runtime/Linux/AudioMixerSDL/Private/AudioMixerPlatformSDL.h
+++ b/Engine/Source/Runtime/Linux/AudioMixerSDL/Private/AudioMixerPlatformSDL.h
@@ -28,7 +28,7 @@ namespace Audio
 		bool GetNumOutputDevices(uint32& OutNumOutputDevices) override;
 		bool GetOutputDeviceInfo(const uint32 InDeviceIndex, FAudioPlatformDeviceInfo& OutInfo) override;
 		bool GetDefaultOutputDeviceIndex(uint32& OutDefaultDeviceIndex) const override;
-		bool OpenAudioStream(const FAudioMixerOpenStreamParams& Params) override;
+		virtual bool OpenAudioStream(const FAudioMixerOpenStreamParams& Params) override;
 		bool CloseAudioStream() override;
 		bool StartAudioStream() override;
 		bool StopAudioStream() override;
@@ -45,6 +45,8 @@ namespace Audio
 
 		void HandleOnBufferEnd(uint8* InOutputBuffer, int32 InOutputBufferLength);
 
+		virtual FString GetCurrentDeviceName() const;
+
 	private:
 
 		SDL_AudioDeviceID AudioDeviceID;
diff --git a/Engine/Source/Runtime/Linux/AudioMixerSDL/Private/AudioMixerSDL.cpp b/Engine/Source/Runtime/Linux/AudioMixerSDL/Private/AudioMixerSDL.cpp
--- a/Engine/Source/Runtime/Linux/AudioMixerSDL/Private/AudioMixerSDL.cpp
+++ b/Engine/Source/Runtime/Linux/AudioMixerSDL/Private/AudioMixerSDL.cpp
@@ -5,6 +5,7 @@
 #include "AudioMixer.h"
 #include "AudioMixerDevice.h"
 #include "AudioMixerPlatformSDL.h"
+#include "AudioMixerPlatformDefinesSDL.h"
 
 
 class FAudioMixerModuleSDL : public IAudioDeviceModule
@@ -23,7 +24,7 @@ public:
 
 	virtual FAudioDevice* CreateAudioDevice() override
 	{
-		return new Audio::FMixerDevice(new Audio::FMixerPlatformSDL());
+		return new Audio::FMixerDevice(new Audio::FAudioMixerPlatformSDL());
 	}
 };
 
diff --git a/Engine/Source/Runtime/Linux/AudioMixerSDL/Private/Linux/AudioMixerPlatformDefinesSDL.h b/Engine/Source/Runtime/Linux/AudioMixerSDL/Private/Linux/AudioMixerPlatformDefinesSDL.h
new file mode 100644
--- /dev/null
+++ b/Engine/Source/Runtime/Linux/AudioMixerSDL/Private/Linux/AudioMixerPlatformDefinesSDL.h
@@ -0,0 +1,10 @@
+// Copyright 1998-2019 Epic Games, Inc. All Rights Reserved.
+
+#pragma once
+
+#include "AudioMixerPlatformSDL.h"
+
+namespace Audio
+{
+	using FAudioMixerPlatformSDL = FMixerPlatformSDL;
+}
