Fix HTML5 build for AudioMixer

diff --git a/Engine/Source/Runtime/Linux/AudioMixerSDL/Private/HTML5/AudioMixerPlatformDefinesSDL.h b/Engine/Source/Runtime/Linux/AudioMixerSDL/Private/HTML5/AudioMixerPlatformDefinesSDL.h
new file mode 100644
--- /dev/null
+++ b/Engine/Source/Runtime/Linux/AudioMixerSDL/Private/HTML5/AudioMixerPlatformDefinesSDL.h
@@ -0,0 +1,10 @@
+// Copyright 1998-2019 Epic Games, Inc. All Rights Reserved.
+
+#pragma once
+
+#include "AudioMixerPlatformSDL.h"
+
+namespace Audio
+{
+	using FAudioMixerPlatformSDL = FMixerPlatformSDL;
+}
