Mutex around the OutputBuffer and Closing of the SDL_AudioDevice

diff --git a/Engine/Source/Runtime/Linux/AudioMixerSDL/Private/AudioMixerPlatformSDL.cpp b/Engine/Source/Runtime/Linux/AudioMixerSDL/Private/AudioMixerPlatformSDL.cpp
--- a/Engine/Source/Runtime/Linux/AudioMixerSDL/Private/AudioMixerPlatformSDL.cpp
+++ b/Engine/Source/Runtime/Linux/AudioMixerSDL/Private/AudioMixerPlatformSDL.cpp
@@ -267,7 +267,12 @@ namespace Audio
 
 		if (AudioDeviceID != INDEX_NONE)
 		{
+			FScopeLock ScopedLock(&OutputBufferMutex);
+
 			SDL_CloseAudioDevice(AudioDeviceID);
+
+			OutputBuffer = nullptr;
+			OutputBufferByteLength = 0;
 		}
 
 		AudioStreamInfo.StreamState = EAudioOutputStreamState::Closed;
@@ -317,6 +322,9 @@ namespace Audio
 
 	void FMixerPlatformSDL::SubmitBuffer(const uint8* Buffer)
 	{
+		// Need to prevent the case in which we close down the audio stream leaving this point to potentially corrupt the free'ed pointer
+		FScopeLock ScopedLock(&OutputBufferMutex);
+
 		if (OutputBuffer)
 		{
 			FMemory::Memcpy(OutputBuffer, Buffer, OutputBufferByteLength);
diff --git a/Engine/Source/Runtime/Linux/AudioMixerSDL/Private/AudioMixerPlatformSDL.h b/Engine/Source/Runtime/Linux/AudioMixerSDL/Private/AudioMixerPlatformSDL.h
--- a/Engine/Source/Runtime/Linux/AudioMixerSDL/Private/AudioMixerPlatformSDL.h
+++ b/Engine/Source/Runtime/Linux/AudioMixerSDL/Private/AudioMixerPlatformSDL.h
@@ -47,6 +47,9 @@ namespace Audio
 
 		virtual FString GetCurrentDeviceName() const;
 
+	protected:
+		FCriticalSection OutputBufferMutex;
+
 	private:
 
 		SDL_AudioDeviceID AudioDeviceID;
