remove more HTML5 code from engine

#rnx

diff --git a/Engine/Platforms/HTML5/Config/HTML5Engine.ini b/Engine/Platforms/HTML5/Config/HTML5Engine.ini
--- a/Engine/Platforms/HTML5/Config/HTML5Engine.ini
+++ b/Engine/Platforms/HTML5/Config/HTML5Engine.ini
@@ -9,8 +9,9 @@ bIsHTMl5=true
 AudioDeviceModuleName=ALAudio
 ; This audio backend can be used by running with the -audiomixer commandline argument:
 AudioMixerModuleName=AudioMixerSDL
-; By default, we use the AudioMixerSDL backend.
-UseAudioMixer=true
+
+; LogAudioMixerSDL: Error: SDL_InitSubSystem create failed: -1
+UseAudioMixer=false
 
 ; Defines a platform-specific volume headroom (in dB) for audio to provide better platform consistency with respect to volume levels.
 PlatformHeadroomDB=0
diff --git a/Engine/Platforms/HTML5/Source/Runtime/AudioMixerSDL/Private/HTML5AudioMixerSDLDefines.h b/Engine/Platforms/HTML5/Source/Runtime/AudioMixerSDL/Private/HTML5AudioMixerSDLDefines.h
--- a/Engine/Platforms/HTML5/Source/Runtime/AudioMixerSDL/Private/HTML5AudioMixerSDLDefines.h
+++ b/Engine/Platforms/HTML5/Source/Runtime/AudioMixerSDL/Private/HTML5AudioMixerSDLDefines.h
@@ -6,5 +6,12 @@
 
 namespace Audio
 {
-	using FAudioMixerPlatformSDL = FMixerPlatformSDL;
+	class FAudioMixerPlatformSDL : public FMixerPlatformSDL
+	{
+public:
+		SDL_AudioFormat GetPlatformAudioFormat() { return AUDIO_S16; }
+		Uint8 GetPlatformChannels() { return 2; }
+		EAudioMixerStreamDataFormat::Type GetAudioStreamFormat() { return EAudioMixerStreamDataFormat::Int16; }
+		int32 GetAudioStreamChannelSize() { return sizeof(int16); }
+	};
 }
diff --git a/Engine/Platforms/HTML5/Source/Runtime/Online/SSL/SSL_HTML5.Build.cs b/Engine/Platforms/HTML5/Source/Runtime/Online/SSL/SSL_HTML5.Build.cs
new file mode 100644
--- /dev/null
+++ b/Engine/Platforms/HTML5/Source/Runtime/Online/SSL/SSL_HTML5.Build.cs
@@ -0,0 +1,13 @@
+// Copyright 1998-2019 Epic Games, Inc. All Rights Reserved.
+
+using UnrealBuildTool;
+
+public class SSL_HTML5 : SSL
+{
+	protected virtual bool PlatformSupportsSSL { get { return true; } }
+
+    public SSL_HTML5(ReadOnlyTargetRules Target) : base(Target)
+    {
+        PublicDefinitions.Add("USE_DEFAULT_SSLCERT=1");
+    }
+}
diff --git a/Engine/Platforms/HTML5/Source/Runtime/OpenGLDrv/Private/HTML5OpenGL.cpp b/Engine/Platforms/HTML5/Source/Runtime/OpenGLDrv/Private/HTML5OpenGL.cpp
--- a/Engine/Platforms/HTML5/Source/Runtime/OpenGLDrv/Private/HTML5OpenGL.cpp
+++ b/Engine/Platforms/HTML5/Source/Runtime/OpenGLDrv/Private/HTML5OpenGL.cpp
@@ -232,6 +232,11 @@ bool FHTML5OpenGL::PE_GLSLToDeviceCompatibleGLSL(FAnsiCharArray& GlslCodeOrigina
 	return true;
 }
 
+void FHTML5OpenGL::PE_SetupTextureFormat(void(*SetupTextureFormat)(EPixelFormat, const FOpenGLTextureFormat&))
+{
+	SetupTextureFormat(PF_R8G8B8A8_UINT, FOpenGLTextureFormat(GL_RGBA8UI, GL_RGBA8UI, GL_RGBA_INTEGER, GL_UNSIGNED_BYTE, false, false));
+}
+
 struct FPlatformOpenGLContext
 {
 	GLuint			ViewportFramebuffer;
diff --git a/Engine/Platforms/HTML5/Source/Runtime/OpenGLDrv/Private/HTML5OpenGL.h b/Engine/Platforms/HTML5/Source/Runtime/OpenGLDrv/Private/HTML5OpenGL.h
--- a/Engine/Platforms/HTML5/Source/Runtime/OpenGLDrv/Private/HTML5OpenGL.h
+++ b/Engine/Platforms/HTML5/Source/Runtime/OpenGLDrv/Private/HTML5OpenGL.h
@@ -262,6 +262,7 @@ struct FHTML5OpenGL : public FOpenGLES2
 	// platform extensions overrides
 	static void PE_GetCurrentOpenGLShaderDeviceCapabilities(FOpenGLShaderDeviceCapabilities& Capabilities);
 	static bool PE_GLSLToDeviceCompatibleGLSL(FAnsiCharArray& GlslCodeOriginal, const FString& ShaderName, GLenum TypeEnum, const FOpenGLShaderDeviceCapabilities& Capabilities, FAnsiCharArray& GlslCode);
+	static void PE_SetupTextureFormat(void(*SetupTextureFormat)(EPixelFormat, const FOpenGLTextureFormat&));
 
 	static FORCEINLINE FString GetAdapterName() { return TEXT(""); }
 
diff --git a/Engine/Source/Editor/UnrealEd/Private/CookOnTheFlyServer.cpp b/Engine/Source/Editor/UnrealEd/Private/CookOnTheFlyServer.cpp
--- a/Engine/Source/Editor/UnrealEd/Private/CookOnTheFlyServer.cpp
+++ b/Engine/Source/Editor/UnrealEd/Private/CookOnTheFlyServer.cpp
@@ -1610,8 +1610,7 @@ bool UCookOnTheFlyServer::StartNetworkFileServer( const bool BindAnyPort )
 		NetworkFileServers.Add(TcpFileServer);
 	}
 
-	// cookonthefly server for html5 -- NOTE: if this is crashing COTF servers, please ask for Nick.Shin (via Josh.Adams)
-#if 1
+#if 0 // cookonthefly server via http
 	INetworkFileServer *HttpFileServer = FModuleManager::LoadModuleChecked<INetworkFileSystemModule>("NetworkFileSystem")
 		.CreateNetworkFileServer(true, BindAnyPort ? 0 : -1, NetworkFileDelegateContainer, ENetworkFileServerProtocol::NFSP_Http);
 	if ( HttpFileServer )
diff --git a/Engine/Source/Runtime/AudioMixerCore/Public/AudioMixerTypes.h b/Engine/Source/Runtime/AudioMixerCore/Public/AudioMixerTypes.h
--- a/Engine/Source/Runtime/AudioMixerCore/Public/AudioMixerTypes.h
+++ b/Engine/Source/Runtime/AudioMixerCore/Public/AudioMixerTypes.h
@@ -14,7 +14,7 @@ namespace Audio {
 			AudioOut, 	// PS4
 			CoreAudio, 	// Mac
 			AudioUnit, 	// iOS
-			SDL2,		// Used on Linux, Html5
+			SDL2,		// Linux
 			OpenSLES, 	// Android
 			Switch, 	// Switch
 			Null		// Unknown/not Supported
diff --git a/Engine/Source/Runtime/Engine/Classes/Kismet/GameplayStatics.h b/Engine/Source/Runtime/Engine/Classes/Kismet/GameplayStatics.h
--- a/Engine/Source/Runtime/Engine/Classes/Kismet/GameplayStatics.h
+++ b/Engine/Source/Runtime/Engine/Classes/Kismet/GameplayStatics.h
@@ -1027,7 +1027,7 @@ public:
 
 	/**
 	 * Returns the string name of the current platform, to perform different behavior based on platform. 
-	 * (Platform names include Windows, Mac, IOS, Android, PS4, XboxOne, HTML5, Linux) */
+	 * (Platform names include Windows, Mac, IOS, Android, PS4, XboxOne, Linux) */
 	UFUNCTION(BlueprintPure, Category="Game")
 	static FString GetPlatformName();
 
diff --git a/Engine/Source/Runtime/Engine/Private/InstancedStaticMesh.cpp b/Engine/Source/Runtime/Engine/Private/InstancedStaticMesh.cpp
--- a/Engine/Source/Runtime/Engine/Private/InstancedStaticMesh.cpp
+++ b/Engine/Source/Runtime/Engine/Private/InstancedStaticMesh.cpp
@@ -442,7 +442,6 @@ void FStaticMeshInstanceBuffer::BindInstanceVertexBuffer(const class FVertexFact
 
 void FStaticMeshInstanceData::Serialize(FArchive& Ar)
 {
-	// HTML5 doesn't support half float so we need to convert at cook time.
 	const bool bCookConvertTransformsToFullFloat = Ar.IsCooking() && bUseHalfFloat && !Ar.CookingTarget()->SupportsFeature(ETargetPlatformFeatures::HalfFloatVertexFormat);
 
 	if (bCookConvertTransformsToFullFloat)
diff --git a/Engine/Source/Runtime/Engine/Private/UnrealEngine.cpp b/Engine/Source/Runtime/Engine/Private/UnrealEngine.cpp
--- a/Engine/Source/Runtime/Engine/Private/UnrealEngine.cpp
+++ b/Engine/Source/Runtime/Engine/Private/UnrealEngine.cpp
@@ -9180,7 +9180,7 @@ void UEngine::PerformanceCapture(UWorld* World, const FString& MapName, const FS
 	// can be define by command line -BuildName="ByCustomBuildName" or "CL<changelist>"
 	FString BuildName = GetBuildNameForPerfTesting();
 
-	// e.g. XboxOne, AllDesktop, Android_.., PS4, HTML5 
+	// e.g. XboxOne, AllDesktop, Android_.., PS4
 	FString PlatformName = FPlatformProperties::PlatformName();
 
 	// e.g. D3D11,OpenGL,Vulcan,D3D12
diff --git a/Engine/Source/Runtime/Linux/AudioMixerSDL/Private/AudioMixerPlatformSDL.cpp b/Engine/Source/Runtime/Linux/AudioMixerSDL/Private/AudioMixerPlatformSDL.cpp
--- a/Engine/Source/Runtime/Linux/AudioMixerSDL/Private/AudioMixerPlatformSDL.cpp
+++ b/Engine/Source/Runtime/Linux/AudioMixerSDL/Private/AudioMixerPlatformSDL.cpp
@@ -109,14 +109,8 @@ namespace Audio
 		SDL_AudioSpec DesiredSpec;
 		DesiredSpec.freq = PlatformSettings.SampleRate;
 
-#if PLATFORM_UNIX
-		DesiredSpec.format = AUDIO_F32;
-		DesiredSpec.channels = 6;
-#else
-		// HTML5 supports s16 format only
-		DesiredSpec.format = AUDIO_S16;
-		DesiredSpec.channels = 2;
-#endif
+		DesiredSpec.format = GetPlatformAudioFormat();
+		DesiredSpec.channels = GetPlatformChannels();
 
 		DesiredSpec.samples = PlatformSettings.CallbackBufferFrameSize;
 		DesiredSpec.callback = OnBufferEnd;
@@ -151,12 +145,7 @@ namespace Audio
 		OutInfo.Name = OutInfo.DeviceId;
 		OutInfo.SampleRate = ActualSpec.freq;
 		
-#if PLATFORM_UNIX
-		OutInfo.Format = EAudioMixerStreamDataFormat::Float;
-#else
-		// HTML5 supports s16 format only
-		OutInfo.Format = EAudioMixerStreamDataFormat::Int16;
-#endif
+		OutInfo.Format = GetAudioStreamFormat();
 		OutInfo.NumChannels = ActualSpec.channels;
 
 		// Assume default channel map order, SDL doesn't support us querying it directly
@@ -198,12 +187,7 @@ namespace Audio
 			return false;
 		}
 
-#if PLATFORM_UNIX
-		AudioSpecPrefered.format = AUDIO_F32;
-#else
-		// HTML5 supports s16 format only
-		AudioSpecPrefered.format = AUDIO_S16;
-#endif
+		AudioSpecPrefered.format = GetPlatformAudioFormat();
 		AudioSpecPrefered.freq = Params.SampleRate;
 		AudioSpecPrefered.channels = AudioStreamInfo.DeviceInfo.NumChannels;
 		AudioSpecPrefered.samples = OpenStreamParams.NumFrames;
@@ -238,12 +222,7 @@ namespace Audio
 		check(AudioSpecReceived.samples == OpenStreamParams.NumFrames);
 
 		// Compute the expected output byte length
-#if PLATFORM_UNIX
-		OutputBufferByteLength = OpenStreamParams.NumFrames * AudioStreamInfo.DeviceInfo.NumChannels * sizeof(float);
-#else
-		// HTML5 supports s16 format only
-		OutputBufferByteLength = OpenStreamParams.NumFrames * AudioStreamInfo.DeviceInfo.NumChannels * sizeof(int16);
-#endif
+		OutputBufferByteLength = OpenStreamParams.NumFrames * AudioStreamInfo.DeviceInfo.NumChannels * GetAudioStreamChannelSize();
 		check(OutputBufferByteLength == AudioSpecReceived.size);
 
 		AudioStreamInfo.StreamState = EAudioOutputStreamState::Open;
@@ -429,7 +408,7 @@ namespace Audio
 #if PLATFORM_UNIX
 		return FAudioPlatformSettings::GetPlatformSettings(FPlatformProperties::GetRuntimeSettingsClassName());
 #else
-		// On HTML5 and Windows, use default parameters.
+		// On Windows, use default parameters.
 		return FAudioPlatformSettings();
 #endif
 	}
diff --git a/Engine/Source/Runtime/Linux/AudioMixerSDL/Private/AudioMixerPlatformSDL.h b/Engine/Source/Runtime/Linux/AudioMixerSDL/Private/AudioMixerPlatformSDL.h
--- a/Engine/Source/Runtime/Linux/AudioMixerSDL/Private/AudioMixerPlatformSDL.h
+++ b/Engine/Source/Runtime/Linux/AudioMixerSDL/Private/AudioMixerPlatformSDL.h
@@ -45,6 +45,11 @@ namespace Audio
 
 		void HandleOnBufferEnd(uint8* InOutputBuffer, int32 InOutputBufferLength);
 
+		SDL_AudioFormat GetPlatformAudioFormat() { return AUDIO_F32; }
+		Uint8 GetPlatformChannels() { return 6; }
+		EAudioMixerStreamDataFormat::Type GetAudioStreamFormat() { return EAudioMixerStreamDataFormat::Float; }
+		int32 GetAudioStreamChannelSize() { return sizeof(float); }
+
 	protected:
 		FCriticalSection OutputBufferMutex;
 
diff --git a/Engine/Source/Runtime/Online/SSL/Private/PlatformSslCertificateManager.h b/Engine/Source/Runtime/Online/SSL/Private/PlatformSslCertificateManager.h
--- a/Engine/Source/Runtime/Online/SSL/Private/PlatformSslCertificateManager.h
+++ b/Engine/Source/Runtime/Online/SSL/Private/PlatformSslCertificateManager.h
@@ -10,26 +10,11 @@
 
 #if PLATFORM_WINDOWS
 #include "Windows/WindowsPlatformSslCertificateManager.h"
-#elif PLATFORM_PS4
-#include "SslCertificateManager.h"
-using FPlatformSslCertificateManager = FSslCertificateManager;
-#elif PLATFORM_XBOXONE
-#include "SslCertificateManager.h"
-using FPlatformSslCertificateManager = FSslCertificateManager;
-#elif PLATFORM_MAC
-#include "SslCertificateManager.h"
-using FPlatformSslCertificateManager = FSslCertificateManager;
-#elif PLATFORM_IOS
-#include "SslCertificateManager.h"
-using FPlatformSslCertificateManager = FSslCertificateManager;
 #elif PLATFORM_ANDROID
 #include "Android/AndroidPlatformSslCertificateManager.h"
-#elif defined(__EMSCRIPTEN__)
-#include "SslCertificateManager.h"
-using FPlatformSslCertificateManager = FSslCertificateManager;
 #elif PLATFORM_UNIX
 #include "Unix/UnixPlatformSslCertificateManager.h"
-#elif PLATFORM_SWITCH
+#elif PLATFORM_MAC || PLATFORM_PS4 || PLATFORM_XBOXONE || PLATFORM_IOS || PLATFORM_SWITCH || defined(USE_DEFAULT_SSLCERT)
 #include "SslCertificateManager.h"
 using FPlatformSslCertificateManager = FSslCertificateManager;
 #else
diff --git a/Engine/Source/Runtime/Online/SSL/SSL.Build.cs b/Engine/Source/Runtime/Online/SSL/SSL.Build.cs
--- a/Engine/Source/Runtime/Online/SSL/SSL.Build.cs
+++ b/Engine/Source/Runtime/Online/SSL/SSL.Build.cs
@@ -4,27 +4,33 @@ using UnrealBuildTool;
 
 public class SSL : ModuleRules
 {
+	protected virtual bool PlatformSupportsSSL
+	{
+		get
+		{
+			return
+				Target.Platform == UnrealTargetPlatform.Mac ||
+				Target.Platform == UnrealTargetPlatform.Win32 ||
+				Target.Platform == UnrealTargetPlatform.Win64 ||
+				Target.IsInPlatformGroup(UnrealPlatformGroup.Unix) ||
+	            Target.Platform == UnrealTargetPlatform.IOS ||
+	            Target.Platform == UnrealTargetPlatform.Android ||
+				Target.Platform == UnrealTargetPlatform.Lumin ||
+	            Target.Platform == UnrealTargetPlatform.PS4;
+		}
+	}
+
     public SSL(ReadOnlyTargetRules Target) : base(Target)
     {
         PublicDefinitions.Add("SSL_PACKAGE=1");
 
-		bool bShouldUseModule =
-			Target.Platform == UnrealTargetPlatform.Mac ||
-			Target.Platform == UnrealTargetPlatform.Win32 ||
-			Target.Platform == UnrealTargetPlatform.Win64 ||
-			Target.IsInPlatformGroup(UnrealPlatformGroup.Unix) ||
-            Target.Platform == UnrealTargetPlatform.IOS ||
-            Target.Platform == UnrealTargetPlatform.Android ||
-			Target.Platform == UnrealTargetPlatform.Lumin ||
-            Target.Platform == UnrealTargetPlatform.PS4;
-
 		PrivateDependencyModuleNames.AddRange(
 			new string[] {
 				"Core",
 			}
 		);
 
-		if (bShouldUseModule)
+		if (PlatformSupportsSSL)
 		{
 			PublicDefinitions.Add("WITH_SSL=1");
 
diff --git a/Engine/Source/Runtime/OpenGLDrv/Private/OpenGLDevice.cpp b/Engine/Source/Runtime/OpenGLDrv/Private/OpenGLDevice.cpp
--- a/Engine/Source/Runtime/OpenGLDrv/Private/OpenGLDevice.cpp
+++ b/Engine/Source/Runtime/OpenGLDrv/Private/OpenGLDevice.cpp
@@ -859,7 +859,7 @@ static void InitRHICapabilitiesForGL()
 	GMaxShadowDepthBufferSizeX = FMath::Min<int32>(Value_GL_MAX_RENDERBUFFER_SIZE, 4096); // Limit to the D3D11 max.
 	GMaxShadowDepthBufferSizeY = FMath::Min<int32>(Value_GL_MAX_RENDERBUFFER_SIZE, 4096);
 	GHardwareHiddenSurfaceRemoval = FOpenGL::HasHardwareHiddenSurfaceRemoval();
-	GRHISupportsInstancing = FOpenGL::SupportsInstancing(); // HTML5 supports it with ANGLE_instanced_arrays or WebGL 2.0+. Android supports it with OpenGL ES3.0+
+	GRHISupportsInstancing = FOpenGL::SupportsInstancing();
 	GSupportsTimestampRenderQueries = FOpenGL::SupportsTimestampQueries();
 
 	// It's not possible to create a framebuffer with the backbuffer as the color attachment and a custom renderbuffer as the depth/stencil surface.
@@ -1029,9 +1029,7 @@ static void InitRHICapabilitiesForGL()
 		SetupTextureFormat(PF_G8, FOpenGLTextureFormat(GL_LUMINANCE, GL_LUMINANCE, GL_LUMINANCE, GL_LUMINANCE, GL_LUMINANCE, GL_UNSIGNED_BYTE, false, false));
 		SetupTextureFormat(PF_A8, FOpenGLTextureFormat(GL_ALPHA, GL_ALPHA, GL_ALPHA, GL_ALPHA, GL_ALPHA, GL_UNSIGNED_BYTE, false, false));
 	#endif
-	#ifdef __EMSCRIPTEN__
-		SetupTextureFormat(PF_R8G8B8A8_UINT, FOpenGLTextureFormat(GL_RGBA8UI, GL_RGBA8UI, GL_RGBA_INTEGER, GL_UNSIGNED_BYTE, false, false));
-	#endif
+		FOpenGL::PE_SetupTextureFormat(&SetupTextureFormat); // platform extension
 
 #if PLATFORM_ANDROID
 		if (GMaxRHIFeatureLevel == ERHIFeatureLevel::ES3_1)
@@ -1115,8 +1113,6 @@ static void InitRHICapabilitiesForGL()
 		}
 		else
 		{
-			// WebGL does not support SRGB versions of DXTn texture formats! Run with SRGB formats disabled.  Will need to make sure
-			// sRGB is always emulated if it's needed.
 			SetupTextureFormat( PF_DXT1,	FOpenGLTextureFormat(GL_COMPRESSED_RGBA_S3TC_DXT1_EXT,	GL_COMPRESSED_RGBA_S3TC_DXT1_EXT,		GL_RGBA,	GL_UNSIGNED_BYTE,	true,	false));
 			SetupTextureFormat( PF_DXT3,	FOpenGLTextureFormat(GL_COMPRESSED_RGBA_S3TC_DXT3_EXT,	GL_COMPRESSED_RGBA_S3TC_DXT3_EXT,		GL_RGBA,	GL_UNSIGNED_BYTE,	true,	false));
 			SetupTextureFormat( PF_DXT5,	FOpenGLTextureFormat(GL_COMPRESSED_RGBA_S3TC_DXT5_EXT,	GL_COMPRESSED_RGBA_S3TC_DXT5_EXT,		GL_RGBA,	GL_UNSIGNED_BYTE,	true,	false));
diff --git a/Engine/Source/Runtime/OpenGLDrv/Private/OpenGLRenderTarget.cpp b/Engine/Source/Runtime/OpenGLDrv/Private/OpenGLRenderTarget.cpp
--- a/Engine/Source/Runtime/OpenGLDrv/Private/OpenGLRenderTarget.cpp
+++ b/Engine/Source/Runtime/OpenGLDrv/Private/OpenGLRenderTarget.cpp
@@ -855,7 +855,6 @@ void FOpenGLDynamicRHI::ReadSurfaceDataRaw(FOpenGLContextState& ContextState, FR
 	{
 		// It's a simple int format. OpenGL converts them internally to what we want.
 		glReadPixels(Rect.Min.X, Rect.Min.Y, SizeX, SizeY, GL_BGRA, UGL_ABGR8, TargetBuffer );
-		// @to-do HTML5. 
 	}
 #endif
 
diff --git a/Engine/Source/Runtime/OpenGLDrv/Public/OpenGL.h b/Engine/Source/Runtime/OpenGLDrv/Public/OpenGL.h
--- a/Engine/Source/Runtime/OpenGLDrv/Public/OpenGL.h
+++ b/Engine/Source/Runtime/OpenGLDrv/Public/OpenGL.h
@@ -54,6 +54,8 @@ extern "C" struct FOpenGLShaderDeviceCapabilities;
 // for platform extensions: from OpenGLResources.h
 typedef TArray<ANSICHAR> FAnsiCharArray;
 
+// for platform extensions: from OpenGLDrvPrivate.h
+extern "C" struct FOpenGLTextureFormat;
 
 // Base static class
 class FOpenGLBase
@@ -408,6 +410,7 @@ public:
 	// for platform extensions
 	static void PE_GetCurrentOpenGLShaderDeviceCapabilities(FOpenGLShaderDeviceCapabilities& Capabilities);
 	static bool PE_GLSLToDeviceCompatibleGLSL(FAnsiCharArray& GlslCodeOriginal, const FString& ShaderName, GLenum TypeEnum, const FOpenGLShaderDeviceCapabilities& Capabilities, FAnsiCharArray& GlslCode) UGL_OPTIONAL(false)
+	static void PE_SetupTextureFormat(void(*SetupTextureFormat)(EPixelFormat, const FOpenGLTextureFormat&)) UGL_OPTIONAL_VOID
 
 protected:
 	static GLint MaxTextureImageUnits;
diff --git a/Engine/Source/Runtime/OpenGLDrv/Public/OpenGLES2.h b/Engine/Source/Runtime/OpenGLDrv/Public/OpenGLES2.h
--- a/Engine/Source/Runtime/OpenGLDrv/Public/OpenGLES2.h
+++ b/Engine/Source/Runtime/OpenGLDrv/Public/OpenGLES2.h
@@ -933,8 +933,6 @@ public:
 #define GL_MAX_SAMPLES 0x8D57
 #else
 // In native OpenGL ES 2.0, define to zero things that are not available.
-// In HTML5 however, always query from the browser what the supported
-// values are.
 #define GL_MAX_3D_TEXTURE_SIZE 0	//0x8073
 #define GL_MAX_COLOR_ATTACHMENTS 0	//0x8CDF
 #define GL_MAX_SAMPLES 0	//0x8D57
