AudioMixerSDL - use a platform-specific header.

diff --git a/Engine/Source/Runtime/Linux/AudioMixerSDL/Private/AudioMixerSDL.cpp b/Engine/Source/Runtime/Linux/AudioMixerSDL/Private/AudioMixerSDL.cpp
--- a/Engine/Source/Runtime/Linux/AudioMixerSDL/Private/AudioMixerSDL.cpp
+++ b/Engine/Source/Runtime/Linux/AudioMixerSDL/Private/AudioMixerSDL.cpp
@@ -5,8 +5,7 @@
 #include "AudioMixer.h"
 #include "AudioMixerDevice.h"
 #include "AudioMixerPlatformSDL.h"
-#include "AudioMixerPlatformDefinesSDL.h"
-
+#include COMPILED_PLATFORM_HEADER(AudioMixerSDLDefines.h)
 
 class FAudioMixerModuleSDL : public IAudioDeviceModule
 {
diff --git a/Engine/Source/Runtime/Linux/AudioMixerSDL/Private/Linux/AudioMixerPlatformDefinesSDL.h b/Engine/Source/Runtime/Linux/AudioMixerSDL/Private/Linux/AudioMixerPlatformDefinesSDL.h
deleted file mode 100644
diff --git a/Engine/Source/Runtime/Linux/AudioMixerSDL/Private/Linux/LinuxAudioMixerSDLDefines.h b/Engine/Source/Runtime/Linux/AudioMixerSDL/Private/Linux/LinuxAudioMixerSDLDefines.h
new file mode 100644
--- /dev/null
+++ b/Engine/Source/Runtime/Linux/AudioMixerSDL/Private/Linux/LinuxAudioMixerSDLDefines.h
@@ -0,0 +1,10 @@
+// Copyright 1998-2019 Epic Games, Inc. All Rights Reserved.
+
+#pragma once
+
+#include "AudioMixerPlatformSDL.h"
+
+namespace Audio
+{
+	using FAudioMixerPlatformSDL = FMixerPlatformSDL;
+}
