Fix up min call in SDL code.

diff --git a/Engine/Source/Runtime/Linux/AudioMixerSDL/Private/AudioMixerPlatformSDL.cpp b/Engine/Source/Runtime/Linux/AudioMixerSDL/Private/AudioMixerPlatformSDL.cpp
--- a/Engine/Source/Runtime/Linux/AudioMixerSDL/Private/AudioMixerPlatformSDL.cpp
+++ b/Engine/Source/Runtime/Linux/AudioMixerSDL/Private/AudioMixerPlatformSDL.cpp
@@ -148,7 +148,7 @@ namespace Audio
 		OutInfo.Format = GetAudioStreamFormat();
 
 		ensure(ActualSpec.channels <= AUDIO_MIXER_MAX_OUTPUT_CHANNELS);
-		OutInfo.NumChannels = FMath::Min(ActualSpec.channels, AUDIO_MIXER_MAX_OUTPUT_CHANNELS);
+		OutInfo.NumChannels = FMath::Min<int32>(static_cast<int32>(ActualSpec.channels), AUDIO_MIXER_MAX_OUTPUT_CHANNELS);
 
 		// Assume default channel map order, SDL doesn't support us querying it directly
 		OutInfo.OutputChannelArray.Reset();
