AudioMixerSDL: only override the default device (UE-92340).

diff --git a/Engine/Source/Runtime/Linux/AudioMixerSDL/Private/AudioMixerPlatformSDL.cpp b/Engine/Source/Runtime/Linux/AudioMixerSDL/Private/AudioMixerPlatformSDL.cpp
--- a/Engine/Source/Runtime/Linux/AudioMixerSDL/Private/AudioMixerPlatformSDL.cpp
+++ b/Engine/Source/Runtime/Linux/AudioMixerSDL/Private/AudioMixerPlatformSDL.cpp
@@ -202,13 +202,16 @@ namespace Audio
 			DeviceName = SDL_GetAudioDeviceName(OpenStreamParams.OutputDeviceIndex, 0);
 		}
 
+		// only the default device can be overriden
 		FString CurrentDeviceName = GetCurrentDeviceName();
-		if (CurrentDeviceName.Len() <= 0)
+		if (OpenStreamParams.OutputDeviceIndex != AUDIO_MIXER_DEFAULT_DEVICE_INDEX || CurrentDeviceName.Len() <= 0)
 		{
+			UE_LOG(LogAudioMixerSDL, Log, TEXT("Opening %s audio device (device index %d)"), DeviceName ? ANSI_TO_TCHAR(DeviceName) : TEXT("default"), OpenStreamParams.OutputDeviceIndex);
 			AudioDeviceID = SDL_OpenAudioDevice(DeviceName, 0, &AudioSpecPrefered, &AudioSpecReceived, 0);
 		}
 		else
 		{
+			UE_LOG(LogAudioMixerSDL, Log, TEXT("Opening overridden '%s' audio device (device index %d)"), *CurrentDeviceName, OpenStreamParams.OutputDeviceIndex);
 			AudioDeviceID = SDL_OpenAudioDevice(TCHAR_TO_ANSI(*CurrentDeviceName), 0, &AudioSpecPrefered, &AudioSpecReceived, 0);
 		}
 
