Fixes Linux crash on startup trying to create and compress audio

diff --git a/Engine/Source/Runtime/Linux/AudioMixerSDL/Private/AudioMixerPlatformSDL.cpp b/Engine/Source/Runtime/Linux/AudioMixerSDL/Private/AudioMixerPlatformSDL.cpp
--- a/Engine/Source/Runtime/Linux/AudioMixerSDL/Private/AudioMixerPlatformSDL.cpp
+++ b/Engine/Source/Runtime/Linux/AudioMixerSDL/Private/AudioMixerPlatformSDL.cpp
@@ -382,6 +382,26 @@ namespace Audio
 #endif // WITH_ENGINE
 	}
 
+	ICompressedAudioInfo* FMixerPlatformSDL::CreateCompressedAudioInfo(const FSoundWaveProxyPtr& InSoundWave)
+	{
+#if WITH_ENGINE
+		if (InSoundWave->IsStreaming())
+		{
+			if (InSoundWave->IsSeekableStreaming())
+			{
+				return new FADPCMAudioInfo();
+			}
+
+			return new FOpusAudioInfo();
+		}
+
+		return nullptr;
+#else
+		checkNoEntry();
+		return nullptr;
+#endif // WITH_ENGINE
+	}
+
 	FString FMixerPlatformSDL::GetDefaultDeviceName()
 	{
 		static FString DefaultName(TEXT("Default SDL Audio Device."));
diff --git a/Engine/Source/Runtime/Linux/AudioMixerSDL/Private/AudioMixerPlatformSDL.h b/Engine/Source/Runtime/Linux/AudioMixerSDL/Private/AudioMixerPlatformSDL.h
--- a/Engine/Source/Runtime/Linux/AudioMixerSDL/Private/AudioMixerPlatformSDL.h
+++ b/Engine/Source/Runtime/Linux/AudioMixerSDL/Private/AudioMixerPlatformSDL.h
@@ -37,6 +37,7 @@ namespace Audio
 		FName GetRuntimeFormat(USoundWave* InSoundWave) override;
 		bool HasCompressedAudioInfoClass(USoundWave* InSoundWave) override;
 		ICompressedAudioInfo* CreateCompressedAudioInfo(USoundWave* InSoundWave) override;
+		ICompressedAudioInfo* CreateCompressedAudioInfo(const FSoundWaveProxyPtr& InSoundWave) override;
 		FString GetDefaultDeviceName() override;
 		FAudioPlatformSettings GetPlatformSettings() const override;
 		void ResumeContext() override;
