Fix spaces usage over tabs

diff --git a/Engine/Source/Runtime/Linux/AudioMixerSDL/AudioMixerSDL.Build.cs b/Engine/Source/Runtime/Linux/AudioMixerSDL/AudioMixerSDL.Build.cs
--- a/Engine/Source/Runtime/Linux/AudioMixerSDL/AudioMixerSDL.Build.cs
+++ b/Engine/Source/Runtime/Linux/AudioMixerSDL/AudioMixerSDL.Build.cs
@@ -21,9 +21,9 @@ public class AudioMixerSDL : ModuleRules
 			PrivateIncludePaths.Add("Runtime/Linux/AudioMixerSDL/Private/" + PlatformName);
 		}
 
-        {        	
-			// Bink Audio isn't yet built for other SDL platforms (e.g. QNX, Stadia),
-			// however the module won't provide a lib and will define WITH_BINK_AUDIO
+		{
+			// Bink Audio isn't yet built for other SDL platforms, however
+			// the module won't provide a lib and will define WITH_BINK_AUDIO
 			// to 0 for those platforms.
 			PrivateDependencyModuleNames.Add("BinkAudioDecoder");
 		}
