Allowing SDL backend to be optimized

diff --git a/Engine/Source/Runtime/Linux/AudioMixerSDL/AudioMixerSDL.Build.cs b/Engine/Source/Runtime/Linux/AudioMixerSDL/AudioMixerSDL.Build.cs
--- a/Engine/Source/Runtime/Linux/AudioMixerSDL/AudioMixerSDL.Build.cs
+++ b/Engine/Source/Runtime/Linux/AudioMixerSDL/AudioMixerSDL.Build.cs
@@ -6,8 +6,6 @@ public class AudioMixerSDL : ModuleRules
 {
 	public AudioMixerSDL(ReadOnlyTargetRules Target) : base(Target)
 	{
-        OptimizeCode = CodeOptimization.Never;
-
         PrivateIncludePathModuleNames.Add("TargetPlatform");
 		PublicIncludePaths.Add("Runtime/AudioMixer/Public");
 		PrivateIncludePaths.Add("Runtime/AudioMixer/Private");
