Fix for non-seeking sound waves

diff --git a/Engine/Source/Runtime/Linux/AudioMixerSDL/Private/AudioMixerPlatformSDL.cpp b/Engine/Source/Runtime/Linux/AudioMixerSDL/Private/AudioMixerPlatformSDL.cpp
--- a/Engine/Source/Runtime/Linux/AudioMixerSDL/Private/AudioMixerPlatformSDL.cpp
+++ b/Engine/Source/Runtime/Linux/AudioMixerSDL/Private/AudioMixerPlatformSDL.cpp
@@ -337,13 +337,13 @@ namespace Audio
 
 		if (RuntimeFormat == Audio::NAME_PLATFORM_SPECIFIC)
 		{
-			if (InSoundWave->HasCompressedData(Audio::NAME_OGG))
+			if (InSoundWave->IsStreaming(nullptr))
 			{
-				RuntimeFormat = Audio::NAME_OGG;
+				return Audio::NAME_OPUS;
 			}
 			else
 			{
-				RuntimeFormat = Audio::NAME_OPUS;
+				return Audio::NAME_OGG;
 			}
 		}
 
