Fix for bink issue on stadia

diff --git a/Engine/Source/Runtime/Linux/AudioMixerSDL/Private/AudioMixerPlatformSDL.cpp b/Engine/Source/Runtime/Linux/AudioMixerSDL/Private/AudioMixerPlatformSDL.cpp
--- a/Engine/Source/Runtime/Linux/AudioMixerSDL/Private/AudioMixerPlatformSDL.cpp
+++ b/Engine/Source/Runtime/Linux/AudioMixerSDL/Private/AudioMixerPlatformSDL.cpp
@@ -347,6 +347,14 @@ namespace Audio
 			}
 		}
 
+#if !WITH_BINK_AUDIO
+		// Some SDL platforms don't yet support bink, so fall back to ADPCM
+		if (InRuntimeFormat == Audio::NAME_BINKA)
+		{
+			return Audio::NAME_ADPCM;
+		}
+#endif
+
 		return RuntimeFormat;
 	}
 
