Build fix for linux

diff --git a/Engine/Source/Runtime/Linux/AudioMixerSDL/Private/AudioMixerPlatformSDL.cpp b/Engine/Source/Runtime/Linux/AudioMixerSDL/Private/AudioMixerPlatformSDL.cpp
--- a/Engine/Source/Runtime/Linux/AudioMixerSDL/Private/AudioMixerPlatformSDL.cpp
+++ b/Engine/Source/Runtime/Linux/AudioMixerSDL/Private/AudioMixerPlatformSDL.cpp
@@ -349,7 +349,7 @@ namespace Audio
 
 #if !WITH_BINK_AUDIO
 		// Some SDL platforms don't yet support bink, so fall back to ADPCM
-		if (InRuntimeFormat == Audio::NAME_BINKA)
+		if (RuntimeFormat == Audio::NAME_BINKA)
 		{
 			return Audio::NAME_ADPCM;
 		}
