Fallback to null renderer for SDL mixer if there's no devices.

diff --git a/Engine/Source/Runtime/Linux/AudioMixerSDL/Private/AudioMixerPlatformSDL.cpp b/Engine/Source/Runtime/Linux/AudioMixerSDL/Private/AudioMixerPlatformSDL.cpp
--- a/Engine/Source/Runtime/Linux/AudioMixerSDL/Private/AudioMixerPlatformSDL.cpp
+++ b/Engine/Source/Runtime/Linux/AudioMixerSDL/Private/AudioMixerPlatformSDL.cpp
@@ -181,60 +181,79 @@ namespace Audio
 			return false;
 		}
 
+		uint32 NumOutputDevices = 0;
+		if (!GetNumOutputDevices(NumOutputDevices))
+		{
+			return false;
+		}
+		
 		OpenStreamParams = Params;
 
 		AudioStreamInfo.Reset();
-		AudioStreamInfo.OutputDeviceIndex = OpenStreamParams.OutputDeviceIndex;
+		AudioStreamInfo.AudioMixer = OpenStreamParams.AudioMixer;
 		AudioStreamInfo.NumOutputFrames = OpenStreamParams.NumFrames;
 		AudioStreamInfo.NumBuffers = OpenStreamParams.NumBuffers;
-		AudioStreamInfo.AudioMixer = OpenStreamParams.AudioMixer;
+		AudioStreamInfo.DeviceInfo.SampleRate = OpenStreamParams.SampleRate;
 
-		if (!GetOutputDeviceInfo(AudioStreamInfo.OutputDeviceIndex, AudioStreamInfo.DeviceInfo))
-		{
-			return false;
-		}
+		// If there's an available device, open it.
+		if (NumOutputDevices > 0)
+		{		
+			if (!GetOutputDeviceInfo(AudioStreamInfo.OutputDeviceIndex, AudioStreamInfo.DeviceInfo))
+			{
+				return false;
+			}
 
-		AudioSpecPrefered.format = GetPlatformAudioFormat();
-		AudioSpecPrefered.freq = Params.SampleRate;
-		AudioSpecPrefered.channels = AudioStreamInfo.DeviceInfo.NumChannels;
-		AudioSpecPrefered.samples = OpenStreamParams.NumFrames;
-		AudioSpecPrefered.callback = OnBufferEnd;
-		AudioSpecPrefered.userdata = (void*)this;
+			AudioStreamInfo.OutputDeviceIndex = OpenStreamParams.OutputDeviceIndex;
 
-		const char* DeviceName = nullptr;
-		if (OpenStreamParams.OutputDeviceIndex != AUDIO_MIXER_DEFAULT_DEVICE_INDEX && OpenStreamParams.OutputDeviceIndex < (uint32)SDL_GetNumAudioDevices(0))
-		{
-			DeviceName = SDL_GetAudioDeviceName(OpenStreamParams.OutputDeviceIndex, 0);
-		}
+			AudioSpecPrefered.format = GetPlatformAudioFormat();
+			AudioSpecPrefered.freq = Params.SampleRate;
+			AudioSpecPrefered.channels = AudioStreamInfo.DeviceInfo.NumChannels;
+			AudioSpecPrefered.samples = OpenStreamParams.NumFrames;
+			AudioSpecPrefered.callback = OnBufferEnd;
+			AudioSpecPrefered.userdata = (void*)this;
 
-		// only the default device can be overriden
-		FString CurrentDeviceName = GetCurrentDeviceName();
-		if (OpenStreamParams.OutputDeviceIndex != AUDIO_MIXER_DEFAULT_DEVICE_INDEX || CurrentDeviceName.Len() <= 0)
-		{
-			UE_LOG(LogAudioMixerSDL, Log, TEXT("Opening %s audio device (device index %d)"), DeviceName ? ANSI_TO_TCHAR(DeviceName) : TEXT("default"), OpenStreamParams.OutputDeviceIndex);
-			AudioDeviceID = SDL_OpenAudioDevice(DeviceName, 0, &AudioSpecPrefered, &AudioSpecReceived, 0);
-		}
-		else
-		{
-			UE_LOG(LogAudioMixerSDL, Log, TEXT("Opening overridden '%s' audio device (device index %d)"), *CurrentDeviceName, OpenStreamParams.OutputDeviceIndex);
-			AudioDeviceID = SDL_OpenAudioDevice(TCHAR_TO_ANSI(*CurrentDeviceName), 0, &AudioSpecPrefered, &AudioSpecReceived, 0);
-		}
+			const char* DeviceName = nullptr;
+			if (OpenStreamParams.OutputDeviceIndex != AUDIO_MIXER_DEFAULT_DEVICE_INDEX && OpenStreamParams.OutputDeviceIndex < (uint32)SDL_GetNumAudioDevices(0))
+			{
+				DeviceName = SDL_GetAudioDeviceName(OpenStreamParams.OutputDeviceIndex, 0);
+			}
 
-		if (!AudioDeviceID)
-		{
-			const char* ErrorText = SDL_GetError();
-			UE_LOG(LogAudioMixerSDL, Error, TEXT("%s"), ANSI_TO_TCHAR(ErrorText));
-			return false;
-		}
+			// only the default device can be overriden
+			FString CurrentDeviceNameString = GetCurrentDeviceName();
+			if (OpenStreamParams.OutputDeviceIndex != AUDIO_MIXER_DEFAULT_DEVICE_INDEX || CurrentDeviceNameString.Len() <= 0)
+			{
+				UE_LOG(LogAudioMixerSDL, Log, TEXT("Opening %s audio device (device index %d)"), DeviceName ? ANSI_TO_TCHAR(DeviceName) : TEXT("default"), OpenStreamParams.OutputDeviceIndex);
+				AudioDeviceID = SDL_OpenAudioDevice(DeviceName, 0, &AudioSpecPrefered, &AudioSpecReceived, 0);
+			}
+			else
+			{
+				UE_LOG(LogAudioMixerSDL, Log, TEXT("Opening overridden '%s' audio device (device index %d)"), *CurrentDeviceNameString, OpenStreamParams.OutputDeviceIndex);
+				AudioDeviceID = SDL_OpenAudioDevice(TCHAR_TO_ANSI(*CurrentDeviceNameString), 0, &AudioSpecPrefered, &AudioSpecReceived, 0);
+			}
 
-		// Make sure our device initialized as expected, we should have already filtered this out before this point.
-		check(AudioSpecReceived.channels == AudioSpecPrefered.channels);
-		check(AudioSpecReceived.samples == OpenStreamParams.NumFrames);
+			if (!AudioDeviceID)
+			{
+				const char* ErrorText = SDL_GetError();
+				UE_LOG(LogAudioMixerSDL, Error, TEXT("%s"), ANSI_TO_TCHAR(ErrorText));
+				return false;
+			}
 
-		// Compute the expected output byte length
-		OutputBufferByteLength = OpenStreamParams.NumFrames * AudioStreamInfo.DeviceInfo.NumChannels * GetAudioStreamChannelSize();
-		check(OutputBufferByteLength == AudioSpecReceived.size);
+			// Make sure our device initialized as expected, we should have already filtered this out before this point.
+			check(AudioSpecReceived.channels == AudioSpecPrefered.channels);
+			check(AudioSpecReceived.samples == OpenStreamParams.NumFrames);
 
+			// Compute the expected output byte length
+			OutputBufferByteLength = OpenStreamParams.NumFrames * AudioStreamInfo.DeviceInfo.NumChannels * GetAudioStreamChannelSize();
+			check(OutputBufferByteLength == AudioSpecReceived.size);
+		}
+		else
+		{
+			// No devices available, Switch to NULL (Silent) Output.
+			AudioStreamInfo.DeviceInfo.OutputChannelArray = { EAudioMixerChannel::FrontLeft, EAudioMixerChannel::FrontRight };
+			AudioStreamInfo.DeviceInfo.NumChannels = 2;			
+			AudioStreamInfo.DeviceInfo.Format = EAudioMixerStreamDataFormat::Float;
+		}
+		
 		AudioStreamInfo.StreamState = EAudioOutputStreamState::Open;
 
 		return true;
@@ -276,8 +295,17 @@ namespace Audio
 		// Start generating audio
 		BeginGeneratingAudio();
 
-		// Unpause audio device to start it rendering audio
-		SDL_PauseAudioDevice(AudioDeviceID, 0);
+		// If we're not using the Null renderer AudioDeviceID will be set to something other than INDEX_NONE.
+		if (AudioDeviceID != INDEX_NONE)
+		{
+			// Unpause audio device to start it rendering audio
+			SDL_PauseAudioDevice(AudioDeviceID, 0);
+		}
+		else
+		{
+			check(!bIsUsingNullDevice);
+			StartRunningNullDevice();
+		}
 
 		AudioStreamInfo.StreamState = EAudioOutputStreamState::Running;
 
@@ -288,8 +316,16 @@ namespace Audio
 	{
 		if (AudioStreamInfo.StreamState != EAudioOutputStreamState::Stopped && AudioStreamInfo.StreamState != EAudioOutputStreamState::Closed)
 		{
-			// Pause the audio device
-			SDL_PauseAudioDevice(AudioDeviceID, 1);
+			if (AudioDeviceID != INDEX_NONE)
+			{
+				// Pause the audio device
+				SDL_PauseAudioDevice(AudioDeviceID, 1);
+			}
+			else
+			{
+				check(bIsUsingNullDevice);
+				StopRunningNullDevice();
+			}
 
 			if (AudioStreamInfo.StreamState == EAudioOutputStreamState::Running)
 			{
