Fix crash on SDL when Recycle threads is enabled

diff --git a/Engine/Source/Runtime/Linux/AudioMixerSDL/Private/AudioMixerPlatformSDL.cpp b/Engine/Source/Runtime/Linux/AudioMixerSDL/Private/AudioMixerPlatformSDL.cpp
--- a/Engine/Source/Runtime/Linux/AudioMixerSDL/Private/AudioMixerPlatformSDL.cpp
+++ b/Engine/Source/Runtime/Linux/AudioMixerSDL/Private/AudioMixerPlatformSDL.cpp
@@ -67,6 +67,13 @@ namespace Audio
 		const char* DriverName = SDL_GetCurrentAudioDriver();
 		UE_LOG(LogAudioMixerSDL, Display, TEXT("Initialized SDL using %s platform API backend."), ANSI_TO_TCHAR(DriverName));
 
+		if (IAudioMixer::ShouldRecycleThreads())
+		{
+			// Pre-create the null render device thread so we can simple wake it up when we need it.
+			// Give it nothing to do, with a slow tick as the default, but ask it to wait for a signal to wake up.
+			CreateNullDeviceThread([] {}, 1.0f, true);
+		}
+
 		bInitialized = true;
 		return true;
 	}
