Fix SDL3 off screen rendering and audio options.

#rnx

diff --git a/Engine/Source/Runtime/ApplicationCore/Private/Linux/LinuxPlatformApplicationMisc.cpp b/Engine/Source/Runtime/ApplicationCore/Private/Linux/LinuxPlatformApplicationMisc.cpp
--- a/Engine/Source/Runtime/ApplicationCore/Private/Linux/LinuxPlatformApplicationMisc.cpp
+++ b/Engine/Source/Runtime/ApplicationCore/Private/Linux/LinuxPlatformApplicationMisc.cpp
@@ -320,8 +320,8 @@ bool FLinuxPlatformApplicationMisc::InitSDL()
 		// If we're rendering offscreen, use the "dummy" SDL video driver
 		if (FParse::Param(FCommandLine::Get(), TEXT("RenderOffScreen")) && !getenv("SDL_VIDEODRIVER"))
 		{
-			UE_LOG(LogInit, Log, TEXT("Hinting SDL to use 'dummy' video driver."), *EglDeviceHint);
-			setenv("SDL_VIDEODRIVER", "dummy", 1);
+			UE_LOG(LogInit, Log, TEXT("Hinting SDL to use 'dummy' video driver."));
+			SDL_SetHint(SDL_HINT_VIDEO_DRIVER, "dummy");
 		}
 
 		// do we need SDL_INIT_HAPTIC, SDL_INIT_SENSOR or SDL_INIT_CAMERA here?
diff --git a/Engine/Source/Runtime/Linux/AudioMixerSDL/Private/AudioMixerPlatformSDL.cpp b/Engine/Source/Runtime/Linux/AudioMixerSDL/Private/AudioMixerPlatformSDL.cpp
--- a/Engine/Source/Runtime/Linux/AudioMixerSDL/Private/AudioMixerPlatformSDL.cpp
+++ b/Engine/Source/Runtime/Linux/AudioMixerSDL/Private/AudioMixerPlatformSDL.cpp
@@ -6,6 +6,7 @@
 #include "AudioMixerDevice.h"
 #include "CoreGlobals.h"
 #include "Misc/ConfigCacheIni.h"
+#include "Misc/CommandLine.h"
 #include "AudioDevice.h"
 
 #if WITH_ENGINE
@@ -62,6 +63,14 @@ namespace Audio
 			return false;
 		}
 
+		// If we're rendering offscreen, use the "dummy" SDL audio driver
+		if (FParse::Param(FCommandLine::Get(), TEXT("RenderOffScreen")) && !getenv("SDL_AUDIODRIVER"))
+		{
+			UE_LOG(LogAudioMixerSDL, Log, TEXT("Hinting SDL to use 'dummy' audio driver."));
+			SDL_SetHint(SDL_HINT_AUDIO_DRIVER, "dummy");
+		}
+
+
 		int32 Result = SDL_InitSubSystem(SDL_INIT_AUDIO);
 		if (Result < 0)
 		{
